New FastGeoStreaming Plugin (Experimental)
- Provides a WorldPartitionCellTransformer that converts immutable static geometry (StaticMeshes and InstanceStaticMeshes, with/without collision) from the world into a lightweight data structure (non-UObject).
- Provides the runtime implementation that asynchronously streams-in and streams-out the geometry (most of the code runs outside of the GameThread).
- Is part of the Level Streaming process (is compatible with World Partition features like Data Layers & HLODs)
- Requires p.Chaos.EnableAsyncInitBody to be enabled.

diff --git a/Engine/Plugins/Experimental/FastGeoStreaming/FastGeoStreaming.uplugin b/Engine/Plugins/Experimental/FastGeoStreaming/FastGeoStreaming.uplugin
new file mode 100644
--- /dev/null
+++ b/Engine/Plugins/Experimental/FastGeoStreaming/FastGeoStreaming.uplugin
@@ -0,0 +1,24 @@
+{
+	"FileVersion": 3,
+	"Version": 1,
+	"VersionName": "0.1",
+	"FriendlyName": "FastGeo Streaming",
+	"Description": "A system that extracts and converts a partitioned world's geometry to optimize world streaming performance.",
+	"Category": "World Building",
+	"CreatedBy": "Epic Games, Inc.",
+	"CreatedByURL": "https://epicgames.com",
+	"EnabledByDefault": false,
+	"CanContainContent": true,
+	"IsBetaVersion": false,
+	"IsExperimentalVersion": true,
+	"Installed": false,
+	"Modules": [
+		{
+			"Name": "FastGeoStreaming",
+			"Type": "Runtime",
+			"LoadingPhase": "PreDefault"
+		}
+	],
+	"Plugins": [
+	]
+}
\ No newline at end of file
diff --git a/Engine/Plugins/Experimental/FastGeoStreaming/Source/FastGeoStreaming/FastGeoStreaming.Build.cs b/Engine/Plugins/Experimental/FastGeoStreaming/Source/FastGeoStreaming/FastGeoStreaming.Build.cs
new file mode 100644
--- /dev/null
+++ b/Engine/Plugins/Experimental/FastGeoStreaming/Source/FastGeoStreaming/FastGeoStreaming.Build.cs
@@ -0,0 +1,22 @@
+// Copyright Epic Games, Inc. All Rights Reserved.
+
+namespace UnrealBuildTool.Rules
+{
+	public class FastGeoStreaming : ModuleRules
+	{
+		public FastGeoStreaming(ReadOnlyTargetRules Target) : base(Target)
+		{
+			PrivateDependencyModuleNames.AddRange(
+				new string[]
+				{
+					"Core",
+					"CoreUObject",
+					"Engine",
+					"PhysicsCore",
+					"RenderCore",
+					"RHI"
+				}
+			);
+		}
+	}
+}
\ No newline at end of file
diff --git a/Engine/Plugins/Experimental/FastGeoStreaming/Source/FastGeoStreaming/Private/FastGeoAsyncRenderStateJobQueue.cpp b/Engine/Plugins/Experimental/FastGeoStreaming/Source/FastGeoStreaming/Private/FastGeoAsyncRenderStateJobQueue.cpp
new file mode 100644
--- /dev/null
+++ b/Engine/Plugins/Experimental/FastGeoStreaming/Source/FastGeoStreaming/Private/FastGeoAsyncRenderStateJobQueue.cpp
@@ -0,0 +1,218 @@
+// Copyright Epic Games, Inc. All Rights Reserved.
+
+#include "FastGeoAsyncRenderStateJobQueue.h"
+#include "FastGeoContainer.h"
+#include "Engine/World.h"
+#include "Engine/Level.h"
+#include "SceneInterface.h"
+
+FFastGeoAsyncRenderStateJobQueue::FFastGeoAsyncRenderStateJobQueue(FSceneInterface* InScene)
+	: Scene(InScene)
+{
+	Scene->GetWorld()->OnAllLevelsChanged().AddRaw(this, &FFastGeoAsyncRenderStateJobQueue::OnUpdateLevelStreamingDone);
+}
+
+FFastGeoAsyncRenderStateJobQueue::~FFastGeoAsyncRenderStateJobQueue()
+{
+	Scene->GetWorld()->OnAllLevelsChanged().RemoveAll(this);
+	Launch();
+	WaitForAsyncTasksExecution();
+	OnAsyncTasksExecuted();
+	check(IsCompleted());
+}
+
+void FFastGeoAsyncRenderStateJobQueue::AddJob(const FJob& Job)
+{
+	TRACE_CPUPROFILER_EVENT_SCOPE(FFastGeoAsyncRenderStateJobQueue::AddJob);
+	check(IsInGameThread());
+
+	if (!PendingJobs.IsValid())
+	{
+		PendingJobs = MakeUnique<FJobSet>();
+	}
+	PendingJobs->Add(Job);
+
+	check(::IsValid(Job.FastGeo));
+	{
+		switch (Job.Type)
+		{
+		case EJobType::CreateRenderState:
+			Job.FastGeo->OnCreateRenderStateBegin_GameThread();
+			break;
+		case EJobType::DestroyRenderState:
+			Job.FastGeo->OnDestroyRenderStateBegin_GameThread();
+			break;
+		}
+	}
+}
+
+void FFastGeoAsyncRenderStateJobQueue::OnAsyncTasksExecuted()
+{
+	check(AreAsyncTasksExecuted());
+	for (const FJobs& Jobs : PipedJobs)
+	{
+		Jobs->OnPostExecute_GameThread();
+	}
+	PipedJobs.Reset();
+	PipeTasks.Reset();
+	IsReadyToRunAsyncTasksEvent.Reset();
+}
+
+bool FFastGeoAsyncRenderStateJobQueue::AreAsyncTasksExecuted() const
+{
+	return !PendingJobs.IsValid() && !Pipe.HasWork();
+}
+
+bool FFastGeoAsyncRenderStateJobQueue::IsCompleted() const
+{
+	if (!PendingJobs.IsValid() && PipedJobs.IsEmpty() && PipeTasks.IsEmpty())
+	{
+		check(AreAsyncTasksExecuted());
+		return true;
+	}
+	return false;
+}
+
+void FFastGeoAsyncRenderStateJobQueue::Launch()
+{
+	TRACE_CPUPROFILER_EVENT_SCOPE(FFastGeoAsyncRenderStateJobQueue::Launch);
+	check(IsInGameThread());
+
+	if (!PendingJobs.IsValid() || PendingJobs->IsEmpty())
+	{
+		return;
+	}
+
+	// Move Jobs to the async queue
+	FJobs& NewJobs = PipedJobs.Emplace_GetRef();
+	NewJobs = MoveTemp(PendingJobs);
+	if (!IsReadyToRunAsyncTasksEvent)
+	{
+		IsReadyToRunAsyncTasksEvent.Emplace(UE_SOURCE_LOCATION);
+	}
+	// Launch a new task in the pipe
+	PipeTasks.Emplace(Pipe.Launch(UE_SOURCE_LOCATION, [this, NewJobs = NewJobs.Get()]() { NewJobs->Execute(); }, *IsReadyToRunAsyncTasksEvent, UE::Tasks::ETaskPriority::BackgroundHigh));
+}
+
+void FFastGeoAsyncRenderStateJobQueue::OnUpdateLevelStreamingDone()
+{
+	TriggerIsReadyToRunAsyncTasksEvent();
+
+}
+
+void FFastGeoAsyncRenderStateJobQueue::TriggerIsReadyToRunAsyncTasksEvent()
+{
+	if (IsReadyToRunAsyncTasksEvent)
+	{
+		IsReadyToRunAsyncTasksEvent->Trigger();
+		IsReadyToRunAsyncTasksEvent.Reset();
+	}
+}
+
+void FFastGeoAsyncRenderStateJobQueue::WaitForAsyncTasksExecution()
+{
+	TRACE_CPUPROFILER_EVENT_SCOPE(FFastGeoAsyncRenderStateJobQueue::WaitForAsyncTasksExecution);
+	check(IsInGameThread());
+
+	TriggerIsReadyToRunAsyncTasksEvent();
+	UE::Tasks::Wait(PipeTasks);
+}
+
+void FFastGeoAsyncRenderStateJobQueue::Tick(bool bForceWaitCompletion)
+{
+	TRACE_CPUPROFILER_EVENT_SCOPE(FFastGeoAsyncRenderStateJobQueue::Tick);
+	check(IsInGameThread());
+
+	// Flush all queued jobs
+	Launch();
+
+	// Wait for tasks if the world is inside a blocking load
+	if (bForceWaitCompletion ||
+		Scene->GetWorld()->GetIsInBlockTillLevelStreamingCompleted() ||
+		Scene->GetWorld()->GetShouldForceUnloadStreamingLevels() ||
+		Scene->GetWorld()->IsBeingCleanedUp())
+	{
+		WaitForAsyncTasksExecution();
+	}
+
+	// Process completed tasks
+	check(PipeTasks.Num() == PipedJobs.Num());
+	for (int i = 0; i < PipeTasks.Num();)
+	{
+		if (PipeTasks[i].IsCompleted())
+		{
+			PipedJobs[i]->OnPostExecute_GameThread();
+			PipeTasks.RemoveAt(i, EAllowShrinking::No);
+			PipedJobs.RemoveAt(i, EAllowShrinking::No);
+		}
+		else
+		{
+			++i;
+		}
+	}
+
+#if DO_CHECK
+	if (AreAsyncTasksExecuted())
+	{
+		check(PipeTasks.IsEmpty());
+		check(PipedJobs.IsEmpty());
+		check(!IsReadyToRunAsyncTasksEvent.IsSet());
+	}
+#endif
+}
+
+void FFastGeoAsyncRenderStateJobQueue::FJobSet::Execute() const
+{
+	TRACE_CPUPROFILER_EVENT_SCOPE(FFastGeoAsyncRenderStateJobQueue::ExecuteJobs);
+	for (const FJob& Job : Jobs)
+	{
+		Job.Execute();
+	}
+}
+
+void FFastGeoAsyncRenderStateJobQueue::FJobSet::OnPostExecute_GameThread() const
+{
+	TRACE_CPUPROFILER_EVENT_SCOPE(FFastGeoAsyncRenderStateJobQueue::OnPostExecute_GameThread);
+	for (const FJob& Job : Jobs)
+	{
+		Job.OnPostExecute_GameThread();
+	}
+}
+
+void FFastGeoAsyncRenderStateJobQueue::FJob::Execute() const
+{
+	if (const bool bPendingKill = FastGeo && !::IsValid(FastGeo); !bPendingKill)
+	{
+		switch (Type)
+		{
+		case EJobType::CreateRenderState:
+		{
+			TRACE_CPUPROFILER_EVENT_SCOPE(FFastGeoAsyncRenderStateJobQueue::CreateRenderState);
+			FastGeo->OnCreateRenderState_Concurrent();
+		}
+		break;
+		case EJobType::DestroyRenderState:
+		{
+			TRACE_CPUPROFILER_EVENT_SCOPE(FFastGeoAsyncRenderStateJobQueue::DestroyRenderState);
+			FastGeo->OnDestroyRenderState_Concurrent();
+		}
+		break;
+		}
+	}
+}
+
+void FFastGeoAsyncRenderStateJobQueue::FJob::OnPostExecute_GameThread() const
+{
+	if (const bool bPendingKill = FastGeo && !::IsValid(FastGeo); !bPendingKill)
+	{
+		switch (Type)
+		{
+		case EJobType::CreateRenderState:
+			FastGeo->OnCreateRenderStateEnd_GameThread();
+			break;
+		case EJobType::DestroyRenderState:
+			FastGeo->OnDestroyRenderStateEnd_GameThread();
+			break;
+		}
+	}
+}
\ No newline at end of file
diff --git a/Engine/Plugins/Experimental/FastGeoStreaming/Source/FastGeoStreaming/Private/FastGeoAsyncRenderStateJobQueue.h b/Engine/Plugins/Experimental/FastGeoStreaming/Source/FastGeoStreaming/Private/FastGeoAsyncRenderStateJobQueue.h
new file mode 100644
--- /dev/null
+++ b/Engine/Plugins/Experimental/FastGeoStreaming/Source/FastGeoStreaming/Private/FastGeoAsyncRenderStateJobQueue.h
@@ -0,0 +1,91 @@
+// Copyright Epic Games, Inc. All Rights Reserved.
+
+#pragma once
+
+#include "Tasks/Task.h"
+#include "Tasks/Pipe.h"
+#include "FastGeoContainer.h"
+
+class FSceneInterface;
+
+struct FFastGeoAsyncRenderStateJobQueue
+{
+	FFastGeoAsyncRenderStateJobQueue(FSceneInterface* Scene);
+	~FFastGeoAsyncRenderStateJobQueue();
+
+	enum class EJobType
+	{
+		CreateRenderState,
+		DestroyRenderState
+	};
+
+	struct FJob
+	{
+		FJob(UFastGeoContainer* InFastGeo, EJobType InType)
+			: FastGeo(InFastGeo)
+			, Type(InType)
+		{}
+
+		bool IsValid() const { return ::IsValid(FastGeo); }
+		void Execute() const;
+		void OnPostExecute_GameThread() const;
+
+		inline bool operator == (const FJob& Other) const
+		{
+			return (FastGeo == Other.FastGeo) && (Type == Other.Type);
+		}
+
+		friend inline uint32 GetTypeHash(const FJob& InJob)
+		{
+			return GetTypeHash(TPair<UFastGeoContainer*, EJobType>{InJob.FastGeo, InJob.Type});
+		}
+
+		UFastGeoContainer* FastGeo;
+		EJobType Type;
+	};
+
+	struct FJobSet
+	{
+		void Add(const FJob& Job)
+		{
+			Jobs.Add(Job);
+		}
+
+		int32 Remove(const FJob& Job)
+		{
+			return Jobs.Remove(Job);
+		}
+
+		bool IsEmpty() const
+		{
+			return Jobs.IsEmpty();
+		}
+
+		void Execute() const;
+		void OnPostExecute_GameThread() const;
+
+	private:
+		TSet<FJob> Jobs;
+	};
+
+	void Tick(bool bForceWaitCompletion);
+	void AddJob(const FJob& Job);
+	bool IsCompleted() const;
+
+private:
+
+	void Launch();
+	void WaitForAsyncTasksExecution();
+	bool AreAsyncTasksExecuted() const;
+	void OnAsyncTasksExecuted();
+	void OnUpdateLevelStreamingDone();
+	void TriggerIsReadyToRunAsyncTasksEvent();
+
+	FSceneInterface* Scene;
+	using FJobs = TUniquePtr<FJobSet>;
+	FJobs PendingJobs;
+	TArray<FJobs> PipedJobs;
+	TArray<UE::Tasks::FTask> PipeTasks;
+	TOptional<UE::Tasks::FTaskEvent> IsReadyToRunAsyncTasksEvent;
+	UE::Tasks::FPipe Pipe{ TEXT("FastGeoAsyncJobQueue") };
+};
\ No newline at end of file
diff --git a/Engine/Plugins/Experimental/FastGeoStreaming/Source/FastGeoStreaming/Private/FastGeoComponent.cpp b/Engine/Plugins/Experimental/FastGeoStreaming/Source/FastGeoStreaming/Private/FastGeoComponent.cpp
new file mode 100644
--- /dev/null
+++ b/Engine/Plugins/Experimental/FastGeoStreaming/Source/FastGeoStreaming/Private/FastGeoComponent.cpp
@@ -0,0 +1,105 @@
+// Copyright Epic Games, Inc. All Rights Reserved.
+
+#include "FastGeoComponent.h"
+#include "FastGeoComponentCluster.h"
+#include "FastGeoContainer.h"
+#include "FastGeoWorldSubsystem.h"
+#include "Engine/Level.h"
+#include "Engine/World.h"
+#include "SceneInterface.h"
+
+const FFastGeoElementType FFastGeoComponent::Type(&IFastGeoElement::Type);
+
+FFastGeoComponent::FFastGeoComponent(int32 InComponentIndex, FFastGeoElementType InType)
+	: Super(InType)
+	, ComponentIndex(InComponentIndex)
+	, Owner(nullptr)
+{
+}
+
+#if WITH_EDITOR
+void FFastGeoComponent::InitializeFromComponent(UActorComponent* Component)
+{
+}
+
+UClass* FFastGeoComponent::GetEditorProxyClass() const
+{
+	return UFastGeoComponentEditorProxy::StaticClass();
+}
+#endif
+
+FArchive& operator<<(FArchive& Ar, FFastGeoComponent& Component)
+{
+	Component.Serialize(Ar);
+	return Ar;
+}
+
+void FFastGeoComponent::SetOwnerComponentCluster(FFastGeoComponentCluster* InOwner)
+{
+	check(InOwner);
+	Owner = InOwner;
+}
+
+FFastGeoComponentCluster* FFastGeoComponent::GetOwnerComponentCluster() const
+{
+	return Owner;
+}
+
+UFastGeoContainer* FFastGeoComponent::GetOwnerContainer() const
+{
+	check(GetOwnerComponentCluster());
+	return GetOwnerComponentCluster()->GetOwnerContainer();
+}
+
+UWorld* FFastGeoComponent::GetWorld() const
+{
+	check(GetOwnerContainer());
+	check(GetOwnerContainer()->GetLevel());
+
+	return GetOwnerContainer()->GetLevel()->OwningWorld;
+}
+
+bool FFastGeoComponent::IsRegistered() const
+{
+	return GetOwnerContainer()->IsRegistered();
+}
+
+FLinearColor FFastGeoComponent::GetDebugColor() const
+{
+	return UFastGeoWorldSubsystem::IsEnableDebugView() ? GetDebugColorImp() : FLinearColor::White;
+}
+
+void FFastGeoComponent::Serialize(FArchive& Ar)
+{
+	Ar << ComponentIndex;
+}
+
+void FFastGeoComponent::OnAsyncCreatePhysicsState()
+{
+	check(IsCollisionEnabled());
+	check(PhysicsStateCreation == EPhysicsStateCreation::NotCreated);
+	PhysicsStateCreation = EPhysicsStateCreation::Creating;
+}
+
+void FFastGeoComponent::OnAsyncCreatePhysicsStateEnd_GameThread()
+{
+	check(PhysicsStateCreation == EPhysicsStateCreation::Creating);
+	PhysicsStateCreation = EPhysicsStateCreation::Created;
+}
+
+void FFastGeoComponent::OnAsyncDestroyPhysicsStateBegin_GameThread()
+{
+	check(PhysicsStateCreation == EPhysicsStateCreation::Created);
+	PhysicsStateCreation = EPhysicsStateCreation::Destroying;
+}
+
+void FFastGeoComponent::OnAsyncDestroyPhysicsState()
+{
+	check(IsCollisionEnabled());
+}
+
+void FFastGeoComponent::OnAsyncDestroyPhysicsStateEnd_GameThread()
+{
+	check(PhysicsStateCreation == EPhysicsStateCreation::Destroying);
+	PhysicsStateCreation = EPhysicsStateCreation::NotCreated;
+}
diff --git a/Engine/Plugins/Experimental/FastGeoStreaming/Source/FastGeoStreaming/Private/FastGeoComponent.h b/Engine/Plugins/Experimental/FastGeoStreaming/Source/FastGeoStreaming/Private/FastGeoComponent.h
new file mode 100644
--- /dev/null
+++ b/Engine/Plugins/Experimental/FastGeoStreaming/Source/FastGeoStreaming/Private/FastGeoComponent.h
@@ -0,0 +1,132 @@
+// Copyright Epic Games, Inc. All Rights Reserved.
+
+#pragma once
+
+#include "CoreMinimal.h"
+#include "IFastGeoElement.h"
+#include "FastGeoElementType.h"
+#include "Physics/Experimental/AsyncPhysicsStateProcessorInterface.h"
+#include "FastGeoComponent.generated.h"
+
+class UWorld;
+class UBodySetup;
+class UActorComponent;
+class FFastGeoComponentCluster;
+class UFastGeoContainer;
+class FSceneInterface;
+class FPrimitiveSceneProxy;
+class FRegisterComponentContext;
+
+class FASTGEOSTREAMING_API FFastGeoComponent : public IFastGeoElement
+{
+public:
+	typedef IFastGeoElement Super;
+
+	/** Static type identifier for this element class */
+	static const FFastGeoElementType Type;
+
+	FFastGeoComponent(int32 InComponentIndex = INDEX_NONE, FFastGeoElementType InType = Type);
+	virtual ~FFastGeoComponent() = default;
+
+#if WITH_EDITOR
+	virtual void InitializeFromComponent(UActorComponent* Component);
+	virtual UClass* GetEditorProxyClass() const;
+#endif
+
+	int32 GetComponentIndex() const { return ComponentIndex; }
+	virtual UBodySetup* GetBodySetup() const { return nullptr; }
+	virtual bool IsCollisionEnabled() const { return false; }
+	virtual void Serialize(FArchive& Ar);
+	virtual void InitializeDynamicProperties() {}
+	virtual void OnAsyncCreatePhysicsState();
+	virtual void OnAsyncDestroyPhysicsStateBegin_GameThread();
+	virtual void OnAsyncDestroyPhysicsState();
+	void OnAsyncCreatePhysicsStateEnd_GameThread();
+	void OnAsyncDestroyPhysicsStateEnd_GameThread();
+
+	FFastGeoComponentCluster* GetOwnerComponentCluster() const;
+	UFastGeoContainer* GetOwnerContainer() const;
+	UWorld* GetWorld() const;
+	bool IsRegistered() const;
+	FLinearColor GetDebugColor() const;
+
+#if WITH_EDITOR
+	void SetEditorProxy(UFastGeoComponentEditorProxy* InComponentEditorProxy)
+	{
+		ComponentEditorProxy = InComponentEditorProxy;
+	}
+		
+	template <typename TEditorProxy>
+	TEditorProxy* GetEditorProxy() const
+	{
+		return Cast<TEditorProxy>(ComponentEditorProxy);
+	}
+#endif
+
+protected:
+	virtual FLinearColor GetDebugColorImp() const { return FLinearColor::White; }
+
+	// Persistent Data
+	int32 ComponentIndex;
+
+	// Transient Data
+	FFastGeoComponentCluster* Owner;
+
+	enum EPhysicsStateCreation
+	{
+		NotCreated,
+		Creating,
+		Created,
+		Destroying
+	};
+
+	EPhysicsStateCreation PhysicsStateCreation = EPhysicsStateCreation::NotCreated;
+
+private:
+	void SetOwnerComponentCluster(FFastGeoComponentCluster* InOwner);
+	friend class FFastGeoComponentCluster;
+
+	friend FArchive& operator<<(FArchive& Ar, FFastGeoComponent& Component);
+
+#if WITH_EDITOR
+	UFastGeoComponentEditorProxy* ComponentEditorProxy;
+#endif
+};
+
+
+UCLASS()
+class UFastGeoComponentEditorProxy : public UObject
+{
+	GENERATED_BODY()
+
+#if WITH_EDITOR
+	friend class UFastGeoContainer;
+
+protected:
+	template <typename TComponent>
+	const TComponent& GetComponent() const
+	{
+		return FastGeoComponent->CastToRef<TComponent>();
+	}
+
+	template <typename TComponent>
+	TComponent& GetComponent()
+	{
+		return FastGeoComponent->CastToRef<TComponent>();
+	}
+
+private:
+	void SetFastGeoComponent(FFastGeoComponent* InFastGeoComponent)
+	{
+		FastGeoComponent = InFastGeoComponent;
+	}
+
+	const FFastGeoComponent* GetFastGeoComponent() const
+	{
+		return FastGeoComponent;
+	}
+
+private:
+	FFastGeoComponent* FastGeoComponent;
+#endif
+};
diff --git a/Engine/Plugins/Experimental/FastGeoStreaming/Source/FastGeoStreaming/Private/FastGeoComponentCluster.cpp b/Engine/Plugins/Experimental/FastGeoStreaming/Source/FastGeoStreaming/Private/FastGeoComponentCluster.cpp
new file mode 100644
--- /dev/null
+++ b/Engine/Plugins/Experimental/FastGeoStreaming/Source/FastGeoStreaming/Private/FastGeoComponentCluster.cpp
@@ -0,0 +1,155 @@
+// Copyright Epic Games, Inc. All Rights Reserved.
+
+#include "FastGeoComponentCluster.h"
+#include "FastGeoStaticMeshComponent.h"
+#include "FastGeoInstancedStaticMeshComponent.h"
+#include "FastGeoContainer.h"
+
+const FFastGeoElementType FFastGeoComponentCluster::Type(&IFastGeoElement::Type);
+
+FFastGeoComponentCluster::FFastGeoComponentCluster(UFastGeoContainer* InOwner, FName InName, FFastGeoElementType InType)
+	: Super(InType)
+	, Owner(InOwner)
+	, Name(InName.ToString())
+	, ComponentClusterIndex(INDEX_NONE)
+{
+}
+
+#if WITH_EDITOR
+FFastGeoComponentCluster::FFastGeoComponentCluster(const FFastGeoComponentCluster& InOther)
+	: Super(InOther.ElementType)
+{
+	ComponentClusterIndex = InOther.ComponentClusterIndex;
+	Owner = InOther.Owner;
+	Name = InOther.Name;
+	StaticMeshComponents = InOther.StaticMeshComponents;
+	InstancedStaticMeshComponents = InOther.InstancedStaticMeshComponents;
+}
+#endif
+
+FArchive& operator<<(FArchive& Ar, FFastGeoComponentCluster& ComponentCluster)
+{
+	ComponentCluster.Serialize(Ar);
+	return Ar;
+}
+
+void FFastGeoComponentCluster::Serialize(FArchive& Ar)
+{
+	Ar << Name;
+	Ar << ComponentClusterIndex;
+	Ar << StaticMeshComponents;
+	Ar << InstancedStaticMeshComponents;
+}
+
+void FFastGeoComponentCluster::InitializeDynamicProperties()
+{
+	ForEachComponent([this](FFastGeoComponent& Component)
+	{
+		Component.SetOwnerComponentCluster(this);
+		Component.InitializeDynamicProperties();
+	});
+}
+
+void FFastGeoComponentCluster::SetOwnerContainer(UFastGeoContainer* InOwner)
+{
+	check(InOwner);
+	Owner = InOwner;
+}
+
+void FFastGeoComponentCluster::SetComponentClusterIndex(int32 InComponentClusterIndex)
+{
+	ComponentClusterIndex = InComponentClusterIndex;
+}
+
+UFastGeoContainer* FFastGeoComponentCluster::GetOwnerContainer() const
+{
+	return Owner;
+}
+
+FFastGeoComponent* FFastGeoComponentCluster::GetComponent(uint32 InComponentTypeID, int32 InComponentIndex)
+{
+	if (FFastGeoInstancedStaticMeshComponent::Type.IsSameTypeID(InComponentTypeID))
+	{
+		return InstancedStaticMeshComponents.IsValidIndex(InComponentIndex) ? &InstancedStaticMeshComponents[InComponentIndex] : nullptr;
+	}
+	else if (FFastGeoStaticMeshComponent::Type.IsSameTypeID(InComponentTypeID))
+	{
+		return StaticMeshComponents.IsValidIndex(InComponentIndex) ? &StaticMeshComponents[InComponentIndex] : nullptr;
+	}
+	check(false);
+	return nullptr;
+}
+
+ULevel* FFastGeoComponentCluster::GetLevel() const
+{
+	return GetOwnerContainer()->GetLevel();
+}
+
+void FFastGeoComponentCluster::UpdateVisibility()
+{
+	TArray<FPrimitiveSceneProxy*, TInlineAllocator<32>> ShowProxies;
+	TArray<FPrimitiveSceneProxy*, TInlineAllocator<32>> HideProxies;
+	ForEachComponent<FFastGeoPrimitiveComponent>([&ShowProxies, &HideProxies](FFastGeoPrimitiveComponent& Component)
+	{
+		const bool bOldIsVisible = Component.IsVisible();
+		Component.UpdateVisibility();
+		const bool bIsVisible = Component.IsVisible();
+		if (bIsVisible != bOldIsVisible)
+		{
+			if (FPrimitiveSceneProxy* SceneProxy = Component.GetSceneProxy())
+			{
+				if (bIsVisible)
+				{
+					ShowProxies.Add(SceneProxy);
+				}
+				else
+				{
+					HideProxies.Add(SceneProxy);
+				}
+			}
+		}
+	});
+
+	if (!ShowProxies.IsEmpty() || !HideProxies.IsEmpty())
+	{
+		ENQUEUE_RENDER_COMMAND()([ShowProxies = MoveTemp(ShowProxies), HideProxies = MoveTemp(HideProxies)](FRHICommandListBase&)
+		{
+			for (FPrimitiveSceneProxy* Proxy : ShowProxies)
+			{
+				Proxy->SetDrawnInGame_RenderThread(true);
+			}
+
+			for (FPrimitiveSceneProxy* Proxy : HideProxies)
+			{
+				Proxy->SetDrawnInGame_RenderThread(false);
+			}
+		});
+	}
+}
+
+#if WITH_EDITOR
+FFastGeoComponent& FFastGeoComponentCluster::AddComponent(FFastGeoElementType InComponentType)
+{
+	FFastGeoComponent* NewComponent = nullptr;
+
+	if (InComponentType == FFastGeoInstancedStaticMeshComponent::Type)
+	{
+		NewComponent = &InstancedStaticMeshComponents.Emplace_GetRef(InstancedStaticMeshComponents.Num());
+	}
+	else if (InComponentType == FFastGeoStaticMeshComponent::Type)
+	{
+		NewComponent = &StaticMeshComponents.Emplace_GetRef(StaticMeshComponents.Num());
+	}
+	else
+	{
+		unimplemented(); // Should never be reached, missing type handling
+	}
+
+	return *NewComponent;
+}
+
+bool FFastGeoComponentCluster::HasComponents() const
+{
+	return InstancedStaticMeshComponents.Num() > 0 || StaticMeshComponents.Num() > 0;
+}
+#endif
\ No newline at end of file
diff --git a/Engine/Plugins/Experimental/FastGeoStreaming/Source/FastGeoStreaming/Private/FastGeoComponentCluster.h b/Engine/Plugins/Experimental/FastGeoStreaming/Source/FastGeoStreaming/Private/FastGeoComponentCluster.h
new file mode 100644
--- /dev/null
+++ b/Engine/Plugins/Experimental/FastGeoStreaming/Source/FastGeoStreaming/Private/FastGeoComponentCluster.h
@@ -0,0 +1,161 @@
+// Copyright Epic Games, Inc. All Rights Reserved.
+
+#pragma once
+
+#include "IFastGeoElement.h"
+#include "FastGeoElementType.h"
+#include "FastGeoStaticMeshComponent.h"
+#include "FastGeoInstancedStaticMeshComponent.h"
+
+class ULevel;
+class UFastGeoContainer;
+class FFastGeoComponent;
+class FArchive;
+
+#if WITH_EDITOR
+#include "UObject/ObjectSaveContext.h"
+#endif // WITH_EDITOR
+
+class FFastGeoComponentCluster : public IFastGeoElement
+{
+public:
+	typedef IFastGeoElement Super;
+
+	/** Static type identifier for this element class */
+	static const FFastGeoElementType Type;
+
+	FFastGeoComponentCluster(UFastGeoContainer* InOwner = nullptr, FName InName = NAME_None, FFastGeoElementType InType = Type);
+
+#if WITH_EDITOR
+	FFastGeoComponentCluster(const FFastGeoComponentCluster& InOther);
+#else
+	FFastGeoComponentCluster(const FFastGeoComponentCluster& InOther) = delete;
+#endif
+
+	virtual ~FFastGeoComponentCluster() {}
+	virtual void OnRegister() {}
+	virtual void OnUnregister() {}
+	virtual void Serialize(FArchive& Ar);
+	virtual bool IsVisible() const { return true; }
+
+	UFastGeoContainer* GetOwnerContainer() const;
+	ULevel* GetLevel() const;
+	FFastGeoComponent* GetComponent(uint32 InComponentTypeID, int32 InComponentIndex);
+	const FString& GetName() const { return Name; }
+	int32 GetComponentClusterIndex() const { return ComponentClusterIndex; }
+	void UpdateVisibility();
+
+#if WITH_EDITOR
+	bool HasComponents() const;
+	virtual FFastGeoComponent& AddComponent(FFastGeoElementType InComponentType);
+	virtual void PreSave(FObjectPreSaveContext ObjectSaveContext) {}
+#endif
+
+	template <typename TComponent = const FFastGeoComponent, typename TFunc>
+	void ForEachComponent(TFunc&& InFunc) const
+	{
+		ForEachComponent<const FFastGeoComponentCluster, const TComponent, TFunc>(this, Forward<TFunc>(InFunc));
+	}
+
+	template <typename TComponent = FFastGeoComponent, typename TFunc>
+	void ForEachComponent(TFunc&& InFunc)
+	{
+		ForEachComponent<FFastGeoComponentCluster, TComponent, TFunc>(this, Forward<TFunc>(InFunc));
+	}
+
+	template <typename TComponent = const FFastGeoComponent, typename TFunc>
+	bool ForEachComponentBreakable(TFunc&& InFunc) const
+	{
+		return ForEachComponentBreakable<const FFastGeoComponentCluster, const TComponent, TFunc>(this, Forward<TFunc>(InFunc));
+	}
+
+	template <typename TComponent = FFastGeoComponent, typename TFunc>
+	bool ForEachComponentBreakable(TFunc&& InFunc)
+	{
+		return ForEachComponentBreakable<FFastGeoComponentCluster, TComponent, TFunc>(this, Forward<TFunc>(InFunc));
+	}
+
+protected:
+	// Transient Data
+	UFastGeoContainer* Owner;
+
+	// Persistent Data
+	FString Name;
+	int32 ComponentClusterIndex;
+	TArray<FFastGeoStaticMeshComponent> StaticMeshComponents;
+	TArray<FFastGeoInstancedStaticMeshComponent> InstancedStaticMeshComponents;
+
+	template <typename ThisType, typename TComponent, typename TFunc>
+	static void ForEachComponent(ThisType* Self, TFunc&& InFunc)
+	{
+		auto ForEachArray = [&InFunc](auto& InArray)
+		{
+			typedef typename TDecay<decltype(InArray)>::Type ArrayType;
+			using TArrayComponent = ArrayType::ElementType;
+
+			if constexpr (std::is_same<TComponent, FFastGeoComponent>::value)
+			{
+				for (auto& Component : InArray)
+				{
+					InFunc(Component);
+				}
+			}
+			else if (TArrayComponent::Type.IsA(TComponent::Type))
+			{
+				for (auto& Component : InArray)
+				{
+					checkSlow(Component.template IsA<TComponent>());
+					Forward<TFunc>(InFunc)(*StaticCast<TComponent*>(&Component));
+				}
+			}
+		};
+		
+		ForEachArray(Self->StaticMeshComponents);
+		ForEachArray(Self->InstancedStaticMeshComponents);
+	}
+
+	template <typename ThisType, typename TComponent, typename TFunc>
+	static bool ForEachComponentBreakable(ThisType* Self, TFunc&& InFunc)
+	{
+		auto ForEachArray = [&InFunc](auto& InArray) -> bool
+		{
+			typedef typename TDecay<decltype(InArray)>::Type ArrayType;
+			using TArrayComponent = ArrayType::ElementType;
+
+			if constexpr (std::is_same<TComponent, FFastGeoComponent>::value)
+			{
+				for (auto& Component : InArray)
+				{
+					if (!InFunc(Component))
+					{
+						return false;
+					}
+				}
+			}
+			else if (TArrayComponent::Type.IsA(TComponent::Type))
+			{
+				for (auto& Component : InArray)
+				{
+					checkSlow(Component.template IsA<TComponent>());
+					if (!Forward<TFunc>(InFunc)(*StaticCast<TComponent*>(&Component)))
+					{
+						return false;
+					}
+				}
+			}
+
+			return true;
+		};
+		
+		return ForEachArray(Self->StaticMeshComponents) &&
+			   ForEachArray(Self->InstancedStaticMeshComponents);
+	}
+
+private:
+	void InitializeDynamicProperties();
+	void SetOwnerContainer(UFastGeoContainer* InOwner);
+	void SetComponentClusterIndex(int32 InComponentClusterIndex);
+	friend class UFastGeoContainer;
+
+	friend FArchive& operator<<(FArchive& Ar, FFastGeoComponentCluster& ComponentCluster);
+};
\ No newline at end of file
diff --git a/Engine/Plugins/Experimental/FastGeoStreaming/Source/FastGeoStreaming/Private/FastGeoContainer.cpp b/Engine/Plugins/Experimental/FastGeoStreaming/Source/FastGeoStreaming/Private/FastGeoContainer.cpp
new file mode 100644
--- /dev/null
+++ b/Engine/Plugins/Experimental/FastGeoStreaming/Source/FastGeoStreaming/Private/FastGeoContainer.cpp
@@ -0,0 +1,615 @@
+// Copyright Epic Games, Inc. All Rights Reserved.
+
+#include "FastGeoContainer.h"
+#include "FastGeoComponentCluster.h"
+#include "FastGeoHLOD.h"
+#include "FastGeoComponent.h"
+#include "FastGeoStreamingManager.h"
+#include "FastGeoLog.h"
+#include "Async/ParallelFor.h"
+#include "Engine/World.h"
+#include "Engine/Level.h"
+#include "PhysicsEngine/BodySetup.h"
+#include "PhysicalMaterials/PhysicalMaterial.h"
+#include "Physics/Experimental/PhysScene_Chaos.h"
+#include "WorldPartition/WorldPartitionHelpers.h"
+#include "WorldPartition/WorldPartition.h"
+#include "PSOPrecache.h"
+
+class FFastGeoAssetRemapArchive : public FArchiveProxy
+{
+public:
+	FFastGeoAssetRemapArchive(FArchive& InArchive, TArray<TObjectPtr<UObject>>& InUniqueAssetsArray)
+		: FArchiveProxy(InArchive)
+		, UniqueAssetsArray(InUniqueAssetsArray)
+	{
+		// For some unknown reason, copy constructor resets ArIsFilterEditorOnly flag copied from the input archive (see FArchiveState(const FArchiveState&))
+		ArIsFilterEditorOnly = InArchive.ArIsFilterEditorOnly;
+		for (int32 Index = 0, Num = UniqueAssetsArray.Num(); Index < Num; ++Index)
+		{
+			UniqueAssets.Add(UniqueAssetsArray[Index], Index);
+		}
+	}
+
+	virtual FArchive& operator<<(UObject*& Obj) override
+	{
+		if (IsLoading())
+		{
+			int32 Index;
+			*this << Index;
+			Obj = UniqueAssetsArray.IsValidIndex(Index) ? UniqueAssetsArray[Index] : nullptr;
+		}
+		else if (IsSaving())
+		{
+			int32 Index = INDEX_NONE;
+			if (int32* ExistingIndex = Obj ? UniqueAssets.Find(Obj) : nullptr)
+			{
+				Index = *ExistingIndex;
+			}
+			*this << Index;
+		}
+		return *this;
+	}
+
+	virtual FArchive& operator<<(FObjectPtr& Obj) override
+	{
+		UObject* ObjPtr = Obj.Get();
+		FArchive& Result = operator<<(ObjPtr);
+		Obj = ObjPtr;
+		return Result;
+	}
+
+	virtual FArchive& operator<<(FLazyObjectPtr& LazyObjectPtr) override { unimplemented(); return *this; }
+	virtual FArchive& operator<<(FSoftObjectPtr& AssetPtr) override { unimplemented(); return *this; }
+	virtual FArchive& operator<<(FSoftObjectPath& AssetPtr) override { unimplemented(); return *this; }
+	virtual FArchive& operator<<(FWeakObjectPtr& Value) override { unimplemented(); return *this; }
+
+private:
+	TMap<UObject*, int32> UniqueAssets;
+	TArray<TObjectPtr<UObject>>& UniqueAssetsArray;
+};
+
+void UFastGeoContainer::Register()
+{
+	TRACE_CPUPROFILER_EVENT_SCOPE(UFastGeoContainer::Register);
+#if WITH_EDITOR
+	if (IsRunningCookCommandlet())
+	{
+		return;
+	}
+#endif
+
+	if (!bIsRegistered)
+	{
+		check(GetWorld()->IsGameWorld());
+#if WITH_EDITOR
+		// In PIE we need to initialize dynamic properties as there's no serialization
+		InitializeDynamicProperties(true);
+#endif
+		UWorldPartition* WorldPartition = FWorldPartitionHelpers::GetWorldPartition(this);
+		const bool bApplyWorldTransform = WorldPartition->HasInstanceTransform();
+		const FTransform& Transform = WorldPartition->GetInstanceTransform();
+		ForEachComponentCluster([bApplyWorldTransform , &Transform](FFastGeoComponentCluster& ComponentCluster)
+		{
+			ComponentCluster.OnRegister();
+
+			if (bApplyWorldTransform)
+			{
+				ComponentCluster.ForEachComponent<FFastGeoPrimitiveComponent>([&Transform](FFastGeoPrimitiveComponent& Component)
+				{
+					Component.ApplyWorldTransform(Transform);
+				});
+			}
+		});
+		FFastGeoStreamingManager::PushAsyncCreateRenderStateJob(this);
+		FFastGeoStreamingManager::PushAsyncCreatePhysicsStateJobs(this);
+		bIsRegistered = true;
+	}
+
+	if (PendingCreate.HasAnyPendingState())
+	{
+		Tick();
+	}
+}
+
+bool UFastGeoContainer::IsRegistered() const
+{
+	return bIsRegistered;
+}
+
+void UFastGeoContainer::Unregister()
+{
+	TRACE_CPUPROFILER_EVENT_SCOPE(UFastGeoContainer::Unregister);
+
+#if WITH_EDITOR
+	if (IsRunningCookCommandlet())
+	{
+		return;
+	}
+#endif
+
+	if (bIsRegistered)
+	{
+		ForEachComponentCluster([](FFastGeoComponentCluster& ComponentCluster)
+		{
+			ComponentCluster.OnUnregister();
+		});
+		FFastGeoStreamingManager::PushAsyncDestroyRenderStateJob(this);
+		FFastGeoStreamingManager::PushAsyncDestroyPhysicsStateJobs(this);
+		bIsRegistered = false;
+	}
+
+	if (PendingDestroy.HasAnyPendingState())
+	{
+		Tick();
+	}
+}
+
+void UFastGeoContainer::Tick(bool bWaitForCompletion)
+{
+	FFastGeoStreamingManager::ProcessAsyncRenderStateJobs(bWaitForCompletion);
+	UWorld* World = GetWorld();
+	check(World);
+	if (FPhysScene* PhysScene = World->GetPhysicsScene())
+	{
+		PhysScene->ProcessAsyncPhysicsStateJobs();
+	}
+}
+
+bool UFastGeoContainer::HasAnyPendingTasks() const
+{
+	return HasAnyPendingCreateTasks() || HasAnyPendingDestroyTasks();
+}
+
+bool UFastGeoContainer::HasAnyPendingCreateTasks() const
+{
+	return PendingCreate.HasAnyPendingState();
+}
+
+bool UFastGeoContainer::HasAnyPendingDestroyTasks() const
+{
+	return PendingDestroy.HasAnyPendingState();
+}
+
+FFastGeoComponentCluster* UFastGeoContainer::GetComponentCluster(uint32 InComponentClusterTypeID, int32 InComponentClusterIndex)
+{
+	if (FFastGeoHLOD::Type.IsSameTypeID(InComponentClusterTypeID))
+	{
+		return HLODs.IsValidIndex(InComponentClusterIndex) ? &HLODs[InComponentClusterIndex] : nullptr;
+	}
+	else if (FFastGeoComponentCluster::Type.IsSameTypeID(InComponentClusterTypeID))
+	{
+		return ComponentClusters.IsValidIndex(InComponentClusterIndex) ? &ComponentClusters[InComponentClusterIndex] : nullptr;
+	}
+	check(false);
+	return nullptr;
+}
+
+ULevel* UFastGeoContainer::GetLevel() const
+{
+	return GetOuterULevel();
+}
+
+UWorld* UFastGeoContainer::GetWorld() const
+{
+	check(GetLevel());
+	return GetLevel()->GetWorld();
+}
+
+void UFastGeoContainer::OnCreateRenderStateBegin_GameThread()
+{
+	check(!PendingCreate.bRenderState);
+	PendingCreate.bRenderState = true;
+}
+
+void UFastGeoContainer::OnDestroyRenderStateBegin_GameThread()
+{
+	check(!PendingDestroy.bRenderState);
+	PendingDestroy.bRenderState = true;
+}
+
+namespace FastGeo
+{
+	static int32 GCreateRenderStateWorkerCount = 0;
+	static FAutoConsoleVariableRef CVarCreateRenderStateWorkerCount(
+		TEXT("FastGeo.CreateRenderStateWorkerCount"),
+		GCreateRenderStateWorkerCount,
+		TEXT("Set the max number of workers to use when creating FastGeo render state. ")
+		TEXT("Only taken into account if value is greater than 1."),
+		ECVF_Default);
+}
+
+void UFastGeoContainer::OnCreateRenderState_Concurrent()
+{
+	TRACE_CPUPROFILER_EVENT_SCOPE(UFastGeoContainer::OnCreateRenderState_Concurrent);
+
+	TArray<FFastGeoPrimitiveComponent*, TInlineAllocator<256>> Components;
+	ForEachComponentCluster([&Components](FFastGeoComponentCluster& ComponentCluster)
+	{
+		ComponentCluster.ForEachComponent<FFastGeoPrimitiveComponent>([&Components](FFastGeoPrimitiveComponent& Component)
+		{
+			if (Component.ShouldCreateRenderState())
+			{
+				Components.Add(&Component);
+			}
+		});
+	});
+
+	static int32 MinNumElementsToProcess = 8;
+	const int32 NumComponents = Components.Num();
+	const bool bIsParallelForAllowed = FApp::ShouldUseThreadingForPerformance() && (NumComponents >= MinNumElementsToProcess);
+	const int32 MinBatchSize = (bIsParallelForAllowed && (FastGeo::GCreateRenderStateWorkerCount > 0) && (NumComponents > FastGeo::GCreateRenderStateWorkerCount)) ? FMath::DivideAndRoundUp(NumComponents, FastGeo::GCreateRenderStateWorkerCount) : 1;
+	
+	ParallelFor(TEXT("UFastGeoContainer::OnCreateRenderState_Concurrent"), NumComponents, MinBatchSize, [&](int32 Index)
+	{
+		// TODO_FASTGEO : Can we call GetUsedMaterials from here? FInstancedStaticMeshSceneProxy -> FStaticMeshSceneProxy -> FStaticMeshSceneProxy::FLODInfo -> FMaterial::MaterialModifiesMeshPosition_RenderThread()
+		FTaskTagScope Scope(ETaskTag::EParallelGameThread);
+		Components[Index]->CreateRenderState(/*Context=*/nullptr);
+	}, bIsParallelForAllowed ? EParallelForFlags::None : EParallelForFlags::ForceSingleThread);
+}
+
+void UFastGeoContainer::OnDestroyRenderState_Concurrent()
+{
+	TRACE_CPUPROFILER_EVENT_SCOPE(UFastGeoContainer::OnDestroyRenderState_Concurrent);
+
+	{
+		FFastGeoPrimitiveComponent::FFastGeoDestroyRenderStateContext Context(GetWorld()->Scene);
+
+		ForEachComponentCluster([&Context](FFastGeoComponentCluster& ComponentCluster)
+		{
+			ComponentCluster.ForEachComponent<FFastGeoPrimitiveComponent>([&Context](FFastGeoPrimitiveComponent& Component)
+			{
+				Component.DestroyRenderState(&Context);
+			});
+		});
+
+#if DO_CHECK
+		bDestroyRenderCommandCompleted = !Context.HasPendingWork();
+#endif
+	}
+
+#if DO_CHECK
+	if (!bDestroyRenderCommandCompleted)
+	{
+		ENQUEUE_RENDER_COMMAND(UFastGeoContainerDestroyed)(
+			[this](FRHICommandListImmediate& RHICmdList)
+			{
+				bDestroyRenderCommandCompleted = true;
+			});
+	}
+#endif
+}
+
+void UFastGeoContainer::OnCreateRenderStateEnd_GameThread()
+{
+	check(PendingCreate.bRenderState);
+	PendingCreate.bRenderState = false;
+}
+
+void UFastGeoContainer::OnDestroyRenderStateEnd_GameThread()
+{
+	check(PendingDestroy.bRenderState);
+	PendingDestroy.bRenderState = false;
+}
+
+void UFastGeoContainer::FinishDestroy()
+{
+#if DO_CHECK
+	check(bDestroyRenderCommandCompleted);
+#endif
+	Super::FinishDestroy();
+}
+
+#if WITH_EDITOR
+void UFastGeoContainer::AddComponentCluster(FFastGeoComponentCluster* ComponentCluster)
+{
+	FFastGeoComponentCluster* NewComponentCluster;
+	if (ComponentCluster->IsA<FFastGeoHLOD>())
+	{
+		NewComponentCluster = &HLODs.Add_GetRef(*ComponentCluster->CastTo<FFastGeoHLOD>());
+		NewComponentCluster->SetComponentClusterIndex(HLODs.Num() - 1);
+	}
+	else
+	{
+		NewComponentCluster = &ComponentClusters.Add_GetRef(*ComponentCluster);
+		NewComponentCluster->SetComponentClusterIndex(ComponentClusters.Num() - 1);
+	}
+}
+
+void UFastGeoContainer::CreateEditorProxyObjects()
+{
+	ForEachComponentCluster([this](FFastGeoComponentCluster& ComponentCluster)
+	{
+		ComponentCluster.ForEachComponent([this](FFastGeoComponent& Component)
+		{
+			if (UClass* EditorProxyClass = Component.GetEditorProxyClass())
+			{
+				UFastGeoComponentEditorProxy* ComponentEditorProxy = NewObject<UFastGeoComponentEditorProxy>(this, EditorProxyClass);
+				ComponentEditorProxy->SetFastGeoComponent(&Component);
+				Component.SetEditorProxy(ComponentEditorProxy);
+				EditorProxyObjects.Add(ComponentEditorProxy);
+			}
+		});
+	});
+}
+
+void UFastGeoContainer::PreSave(FObjectPreSaveContext ObjectSaveContext)
+{
+	Super::PreSave(ObjectSaveContext);
+
+	ForEachComponentCluster([&ObjectSaveContext](FFastGeoComponentCluster& ComponentCluster)
+	{
+		ComponentCluster.PreSave(ObjectSaveContext);
+	});
+}
+
+class FFastGeoGatherFastGeoContainerAssetRefsArchive : public FArchive
+{
+public:
+	FFastGeoGatherFastGeoContainerAssetRefsArchive(UFastGeoContainer& Container)
+	{
+		KnownClasses.Add(URuntimeVirtualTexture::StaticClass());
+		KnownClasses.Add(UStaticMesh::StaticClass());
+		KnownClasses.Add(UMaterialInterface::StaticClass());
+		KnownClasses.Add(UPhysicalMaterial::StaticClass());
+
+		SetIsPersistent(true);
+		SetIsSaving(true);
+		ArIgnoreOuterRef = true;
+		ArIsObjectReferenceCollector = true;
+		ArShouldSkipBulkData = true;
+		Container.SerializeComponentClusters(*this);
+	}
+
+	virtual FArchive& operator<<(UObject*& Obj) override
+	{
+		if (Obj && !Obj->IsTemplate() && !Obj->HasAnyFlags(RF_Transient))
+		{
+			if (!UniqueAssets.Contains(Obj))
+			{
+				check(IsKnownClass(Obj->GetClass()));
+				UniqueAssets.Add(Obj);
+			}
+		}
+		return *this;
+	}
+
+	virtual FArchive& operator<<(FObjectPtr& Obj) override
+	{
+		UObject* ObjPtr = Obj.Get();
+		FArchive& Result = operator<<(ObjPtr);
+		Obj = ObjPtr;
+		return Result;
+	}
+
+	const TSet<UObject*>& GetUniqueAssets() { return UniqueAssets; }
+
+private:
+	bool IsKnownClass(UClass* Class)
+	{
+		for (UClass* KnownClass : KnownClasses)
+		{
+			if (Class->IsChildOf(KnownClass))
+			{
+				return true;
+			}
+		}
+		return false;
+	}
+
+	TSet<UClass*> KnownClasses;
+	TSet<UObject*> UniqueAssets;
+};
+
+void UFastGeoContainer::CollectAssetReferences()
+{
+	Assets = FFastGeoGatherFastGeoContainerAssetRefsArchive(*this).GetUniqueAssets().Array();
+}
+
+void UFastGeoContainer::OnCreated()
+{
+	// Initialize component clusters & components dynamic properties
+	InitializeDynamicProperties(GetLevel()->GetWorld()->IsGameWorld());
+
+	// In editor, we need a UObject representation of the components for some operations
+	CreateEditorProxyObjects();
+
+	// Collect references in order to avoid garbage collection of objects that may now be unreferenced
+	// The fast geo container will hold onto those objects if necessary.
+	CollectAssetReferences();
+}
+#endif
+
+void UFastGeoContainer::Serialize(FArchive& Ar)
+{
+	Super::Serialize(Ar);
+
+	FFastGeoAssetRemapArchive AssetRemapAr(Ar, Assets);
+	SerializeComponentClusters(AssetRemapAr);
+
+#if !WITH_EDITOR
+	if (Ar.IsLoading())
+	{
+		// Once loaded, initialize component clusters & components dynamic properties
+		InitializeDynamicProperties(true);
+	}
+#endif
+}
+
+void UFastGeoContainer::SerializeComponentClusters(FArchive& Ar)
+{
+	Ar << ComponentClusters;
+	Ar << HLODs;
+}
+
+void UFastGeoContainer::PostLoad()
+{
+	TRACE_CPUPROFILER_EVENT_SCOPE(UFastGeoContainer::PostLoad);
+
+	Super::PostLoad();
+
+	UE_CLOG(!FPhysScene::SupportsAsyncPhysicsStateCreation() || !FPhysScene::SupportsAsyncPhysicsStateDestruction(), LogFastGeoStreaming, Error, TEXT("FastGeoStreaming requires 'p.Chaos.EnableAsyncInitBody' to be enabled."));
+
+#if !WITH_EDITOR && UE_WITH_PSO_PRECACHING
+	if (IsComponentPSOPrecachingEnabled())
+	{
+		ForEachComponentCluster([](FFastGeoComponentCluster& ComponentCluster)
+		{
+			ComponentCluster.ForEachComponent<FFastGeoPrimitiveComponent>([](FFastGeoPrimitiveComponent& Component)
+			{
+				Component.PrecachePSOs();
+			});
+		});
+	}
+#endif
+}
+
+void UFastGeoContainer::InitializeDynamicProperties(bool bInitForPlay)
+{
+	ForEachComponentCluster([this](FFastGeoComponentCluster& ComponentCluster)
+	{
+		ComponentCluster.SetOwnerContainer(this);
+		ComponentCluster.InitializeDynamicProperties();
+	});
+
+	if (bInitForPlay)
+	{
+		CollisionComponents.Reset();
+		ForEachComponentCluster([this](FFastGeoComponentCluster& ComponentCluster)
+		{
+			ComponentCluster.ForEachComponent([this](FFastGeoComponent& Component)
+			{
+				if (Component.IsCollisionEnabled())
+				{
+					CollisionComponents.Add(&Component);
+				}
+			});
+		});
+	}
+}
+
+void UFastGeoContainer::OnCreatePhysicsStateBegin_GameThread()
+{
+	TRACE_CPUPROFILER_EVENT_SCOPE(UFastGeoContainer::OnCreatePhysicsStateBegin_GameThread);
+
+	UWorld* World = GetWorld();
+	check(World);
+	FPhysScene* PhysScene = World->GetPhysicsScene();
+	check(PhysScene);
+	PhysScene->PushAsyncCreatePhysicsStateJob(this);
+
+	check(!PendingCreate.bPhysicsState);
+	PendingCreate.bPhysicsState = true;
+}
+
+void UFastGeoContainer::OnDestroyPhysicsStateBegin_GameThread()
+{
+	TRACE_CPUPROFILER_EVENT_SCOPE(UFastGeoContainer::OnDestroyPhysicsStateBegin_GameThread);
+
+	check(!PendingDestroy.bPhysicsState);
+	PendingDestroy.bPhysicsState = true;
+
+	UWorld* World = GetWorld();
+	check(World);
+	FPhysScene* PhysScene = World->GetPhysicsScene();
+	check(PhysScene);
+	verify(PhysScene->PushAsyncDestroyPhysicsStateJob(this));
+}
+
+bool UFastGeoContainer::AllowsAsyncPhysicsStateCreation() const
+{
+	check(FPhysScene::SupportsAsyncPhysicsStateCreation());
+	return true;
+}
+
+bool UFastGeoContainer::AllowsAsyncPhysicsStateDestruction() const
+{
+	check(FPhysScene::SupportsAsyncPhysicsStateDestruction());
+	return true;
+}
+
+bool UFastGeoContainer::IsAsyncPhysicsStateCreated() const
+{
+	return PhysicsStateCreation == EPhysicsStateCreation::Created;
+}
+
+UObject* UFastGeoContainer::GetAsyncPhysicsStateObject() const
+{
+	return const_cast<UFastGeoContainer*>(this);
+}
+
+void UFastGeoContainer::OnAsyncCreatePhysicsState()
+{
+	TRACE_CPUPROFILER_EVENT_SCOPE(UFastGeoContainer::OnAsyncCreatePhysicsState);
+	check(PhysicsStateCreation == EPhysicsStateCreation::NotCreated);
+	PhysicsStateCreation = EPhysicsStateCreation::Creating;
+
+	for (FFastGeoComponent* Component : CollisionComponents)
+	{
+		Component->OnAsyncCreatePhysicsState();
+	}
+}
+
+void UFastGeoContainer::OnAsyncCreatePhysicsStateEnd_GameThread()
+{
+	TRACE_CPUPROFILER_EVENT_SCOPE(UFastGeoContainer::OnAsyncCreatePhysicsStateEnd_GameThread);
+	check(PhysicsStateCreation == EPhysicsStateCreation::Creating);
+	PhysicsStateCreation = EPhysicsStateCreation::Created;
+
+	check(PendingCreate.bPhysicsState);
+	PendingCreate.bPhysicsState = false;
+
+	// This loop could be removed if component's PhysicsStateCreation was removed
+	for (FFastGeoComponent* Component : CollisionComponents)
+	{
+		Component->OnAsyncCreatePhysicsStateEnd_GameThread();
+	}
+}
+
+void UFastGeoContainer::OnAsyncDestroyPhysicsStateBegin_GameThread()
+{
+	TRACE_CPUPROFILER_EVENT_SCOPE(UFastGeoContainer::OnAsyncDestroyPhysicsStateBegin_GameThread);
+	check(PhysicsStateCreation == EPhysicsStateCreation::Created);
+	PhysicsStateCreation = EPhysicsStateCreation::Destroying;
+
+	for (FFastGeoComponent* Component : CollisionComponents)
+	{
+		Component->OnAsyncDestroyPhysicsStateBegin_GameThread();
+	}
+}
+
+void UFastGeoContainer::OnAsyncDestroyPhysicsState()
+{
+	TRACE_CPUPROFILER_EVENT_SCOPE(UFastGeoContainer::OnAsyncDestroyPhysicsState);
+	for (FFastGeoComponent* Component : CollisionComponents)
+	{
+		Component->OnAsyncDestroyPhysicsState();
+	}
+}
+
+void UFastGeoContainer::OnAsyncDestroyPhysicsStateEnd_GameThread()
+{
+	TRACE_CPUPROFILER_EVENT_SCOPE(UFastGeoContainer::OnAsyncDestroyPhysicsStateEnd_GameThread);
+	check(PhysicsStateCreation == EPhysicsStateCreation::Destroying);
+	PhysicsStateCreation = EPhysicsStateCreation::NotCreated;
+
+	check(PendingDestroy.bPhysicsState);
+	PendingDestroy.bPhysicsState = false;
+
+	// This loop could be removed if component's PhysicsStateCreation was removed
+	for (FFastGeoComponent* Component : CollisionComponents)
+	{
+		Component->OnAsyncDestroyPhysicsStateEnd_GameThread();
+	}
+}
+
+void UFastGeoContainer::CollectBodySetupsWithPhysicsMeshesToCreate(TSet<UBodySetup*>& OutBodySetups) const
+{
+	for (FFastGeoComponent* Component : CollisionComponents)
+	{
+		if (UBodySetup* BodySetup = const_cast<FFastGeoComponent*>(Component)->GetBodySetup(); BodySetup && !BodySetup->bCreatedPhysicsMeshes)
+		{
+			OutBodySetups.Add(BodySetup);
+		}
+	}
+}
diff --git a/Engine/Plugins/Experimental/FastGeoStreaming/Source/FastGeoStreaming/Private/FastGeoContainer.h b/Engine/Plugins/Experimental/FastGeoStreaming/Source/FastGeoStreaming/Private/FastGeoContainer.h
new file mode 100644
--- /dev/null
+++ b/Engine/Plugins/Experimental/FastGeoStreaming/Source/FastGeoStreaming/Private/FastGeoContainer.h
@@ -0,0 +1,226 @@
+// Copyright Epic Games, Inc. All Rights Reserved.
+
+#pragma once
+
+#include "Engine/AssetUserData.h"
+#include "Physics/Experimental/AsyncPhysicsStateProcessorInterface.h"
+#include "PhysicsEngine/BodyInstance.h"
+#include "FastGeoComponentCluster.h"
+#include "FastGeoHLOD.h"
+#include "UObject/ObjectPtr.h"
+
+#if WITH_EDITOR
+#include "UObject/ObjectSaveContext.h"
+#endif // WITH_EDITOR
+
+#include "FastGeoContainer.generated.h"
+
+class UWorld;
+class ULevel;
+class FFastGeoComponent;
+class URuntimeVirtualTexture;
+class UStaticMesh;
+class UMaterialInterface;
+
+struct FFastGeoComponentClusterPendingState
+{
+	FFastGeoComponentClusterPendingState()
+		: bRenderState(false)
+	{}
+
+	bool bRenderState;
+	bool bPhysicsState;
+	bool HasAnyPendingState() const { return bRenderState || bPhysicsState; }
+};
+
+UCLASS(Within = Level)
+class FASTGEOSTREAMING_API UFastGeoContainer : public UAssetUserData, public IAsyncPhysicsStateProcessor
+{
+	GENERATED_BODY()
+
+public:
+	void Register();
+	void Unregister();
+	bool IsRegistered() const;
+	void Tick(bool bWaitForCompletion = false);
+	bool HasAnyPendingTasks() const;
+	bool HasAnyPendingCreateTasks() const;
+	bool HasAnyPendingDestroyTasks() const;
+	UWorld* GetWorld() const;
+	ULevel* GetLevel() const;
+	FFastGeoComponentCluster* GetComponentCluster(uint32 InComponentClusterTypeID, int32 InComponentClusterIndex);
+
+	//~ Begin UObject Interface
+	virtual void FinishDestroy() override;
+	virtual void Serialize(FArchive& Ar) override;
+	virtual void PostLoad() override;
+#if WITH_EDITOR
+	virtual void PreSave(FObjectPreSaveContext ObjectSaveContext) override;
+#endif
+	//~ End UObject Interface
+
+#if WITH_EDITOR
+	void OnCreated();
+	void AddComponentCluster(FFastGeoComponentCluster* ComponentCluster);
+#endif
+
+	void InitializeDynamicProperties(bool bInitForPlay = false);
+
+	void OnCreateRenderStateBegin_GameThread();
+	void OnDestroyRenderStateBegin_GameThread();
+	void OnCreateRenderStateEnd_GameThread();
+	void OnDestroyRenderStateEnd_GameThread();
+	void OnCreateRenderState_Concurrent();
+	void OnDestroyRenderState_Concurrent();
+
+	void OnCreatePhysicsStateBegin_GameThread();
+	void OnDestroyPhysicsStateBegin_GameThread();
+
+	//~ Begin IAsyncPhysicsStateProcessor
+	virtual bool AllowsAsyncPhysicsStateCreation() const override final;
+	virtual bool AllowsAsyncPhysicsStateDestruction() const override final;
+	virtual bool IsAsyncPhysicsStateCreated() const override final;
+	virtual UObject* GetAsyncPhysicsStateObject() const override final;
+	virtual void OnAsyncCreatePhysicsState() override;
+	virtual void OnAsyncCreatePhysicsStateEnd_GameThread() override;
+	virtual void OnAsyncDestroyPhysicsStateBegin_GameThread() override;
+	virtual void OnAsyncDestroyPhysicsState() override;
+	virtual void OnAsyncDestroyPhysicsStateEnd_GameThread() override;
+	virtual void CollectBodySetupsWithPhysicsMeshesToCreate(TSet<UBodySetup*>& OutBodySetups) const override;
+	//~ End IAsyncPhysicsStateProcessor
+
+	template <typename TComponentCluster = const FFastGeoComponentCluster, typename TFunc>
+	void ForEachComponentCluster(TFunc&& InFunc) const
+	{
+		ForEachComponentCluster<const UFastGeoContainer, const TComponentCluster, TFunc>(this, Forward<TFunc>(InFunc));
+	}
+
+	template <typename TComponentCluster = FFastGeoComponentCluster, typename TFunc>
+	void ForEachComponentCluster(TFunc&& InFunc)
+	{
+		ForEachComponentCluster<UFastGeoContainer, TComponentCluster, TFunc>(this, Forward<TFunc>(InFunc));
+	}
+
+	template <typename TComponentCluster = const FFastGeoComponentCluster, typename TFunc>
+	bool ForEachComponentClusterBreakable(TFunc&& InFunc) const
+	{
+		return ForEachComponentClusterBreakable<const UFastGeoContainer, const TComponentCluster, TFunc>(this, Forward<TFunc>(InFunc));
+	}
+
+	template <typename TComponentCluster = FFastGeoComponentCluster, typename TFunc>
+	bool ForEachComponentClusterBreakable(TFunc&& InFunc)
+	{
+		return ForEachComponentClusterBreakable<UFastGeoContainer, TComponentCluster, TFunc>(this, Forward<TFunc>(InFunc));
+	}
+
+protected:
+
+	FFastGeoComponentClusterPendingState PendingCreate;
+	FFastGeoComponentClusterPendingState PendingDestroy;
+
+private:
+
+#if WITH_EDITOR
+	void CreateEditorProxyObjects();
+	void CollectAssetReferences();
+#endif
+	void SerializeComponentClusters(FArchive& Ar);
+
+	template <typename ThisType, typename TComponentCluster, typename TFunc>
+	static void ForEachComponentCluster(ThisType* Self, TFunc&& InFunc)
+	{
+		auto ForEachArray = [&InFunc](auto& InArray)
+		{
+			typedef typename TDecay<decltype(InArray)>::Type ArrayType;
+			using TArrayComponentCluster = ArrayType::ElementType;
+
+			if constexpr (std::is_same<TComponentCluster, FFastGeoComponentCluster>::value)
+			{
+				for (auto& ComponentCluster : InArray)
+				{
+					InFunc(ComponentCluster);
+				}
+			}
+			else if (TArrayComponentCluster::Type.IsA(TComponentCluster::Type))
+			{
+				for (auto& ComponentCluster : InArray)
+				{
+					checkSlow(ComponentCluster.template IsA<TComponentCluster>());
+					Forward<TFunc>(InFunc)(*StaticCast<TComponentCluster*>(&ComponentCluster));
+				}
+			}
+		};
+
+		ForEachArray(Self->ComponentClusters);
+		ForEachArray(Self->HLODs);
+	}
+
+	template <typename ThisType, typename TComponentCluster, typename TFunc>
+	static bool ForEachComponentClusterBreakable(ThisType* Self, TFunc&& InFunc)
+	{
+		auto ForEachArray = [&InFunc](auto& InArray) -> bool
+		{
+			typedef typename TDecay<decltype(InArray)>::Type ArrayType;
+			using TArrayComponentCluster = ArrayType::ElementType;
+
+			if constexpr (std::is_same<TComponentCluster, FFastGeoComponentCluster>::value)
+			{
+				for (auto& ComponentCluster : InArray)
+				{
+					if (!InFunc(ComponentCluster))
+					{
+						return false;
+					}
+				}
+			}
+			else if (TArrayComponentCluster::Type.IsA(TComponentCluster::Type))
+			{
+				for (auto& ComponentCluster : InArray)
+				{
+					checkSlow(ComponentCluster.template IsA<TComponentCluster>());
+					if (!Forward<TFunc>(InFunc)(*StaticCast<TComponentCluster*>(&ComponentCluster)))
+					{
+						return false;
+					}
+				}
+			}
+
+			return true;
+		};
+
+		return ForEachArray(Self->ComponentClusters) &&
+			ForEachArray(Self->HLODs);
+	}
+
+	enum EPhysicsStateCreation
+	{
+		NotCreated,
+		Creating,
+		Created,
+		Destroying
+	};
+	EPhysicsStateCreation PhysicsStateCreation = EPhysicsStateCreation::NotCreated;
+
+	// Persistent data
+	TArray<FFastGeoComponentCluster> ComponentClusters;
+	TArray<FFastGeoHLOD> HLODs;
+
+	UPROPERTY()
+	TArray<TObjectPtr<UObject>> Assets;
+
+#if WITH_EDITORONLY_DATA
+	UPROPERTY(Transient)
+	TSet<TObjectPtr<UObject>> EditorProxyObjects;
+#endif
+
+	// Transient data
+	TArray<FFastGeoComponent*> CollisionComponents;
+
+	bool bIsRegistered = false;
+
+#if DO_CHECK
+	TAtomic<bool> bDestroyRenderCommandCompleted = true;
+#endif
+
+	friend class FFastGeoGatherFastGeoContainerAssetRefsArchive;
+};
\ No newline at end of file
diff --git a/Engine/Plugins/Experimental/FastGeoStreaming/Source/FastGeoStreaming/Private/FastGeoElementType.h b/Engine/Plugins/Experimental/FastGeoStreaming/Source/FastGeoStreaming/Private/FastGeoElementType.h
new file mode 100644
--- /dev/null
+++ b/Engine/Plugins/Experimental/FastGeoStreaming/Source/FastGeoStreaming/Private/FastGeoElementType.h
@@ -0,0 +1,34 @@
+// Copyright Epic Games, Inc. All Rights Reserved.
+
+#pragma once
+
+#include "HAL/Platform.h"
+
+struct FFastGeoElementType
+{
+public:
+	explicit FFastGeoElementType(const FFastGeoElementType* Parent = nullptr) : ID(++NextUniqueID), ParentType(Parent) {}
+
+	bool operator==(const FFastGeoElementType& Other) const
+	{
+		return ID == Other.ID || (ParentType != nullptr && *ParentType == Other);
+	}
+
+	bool IsA(const FFastGeoElementType& Other) const
+	{
+		return (ID == Other.ID) || (ParentType && ParentType->IsA(Other));
+	}
+
+	/** Returns true if element has the same type ID than the specified type */
+	bool IsSameTypeID(uint32 InID) const
+	{
+		return ID == InID;
+	}
+
+private:
+	static uint32 NextUniqueID;
+	uint32 ID;
+	const FFastGeoElementType* ParentType;
+
+	friend class IFastGeoElement;
+};
\ No newline at end of file
diff --git a/Engine/Plugins/Experimental/FastGeoStreaming/Source/FastGeoStreaming/Private/FastGeoHLOD.cpp b/Engine/Plugins/Experimental/FastGeoStreaming/Source/FastGeoStreaming/Private/FastGeoHLOD.cpp
new file mode 100644
--- /dev/null
+++ b/Engine/Plugins/Experimental/FastGeoStreaming/Source/FastGeoStreaming/Private/FastGeoHLOD.cpp
@@ -0,0 +1,167 @@
+// Copyright Epic Games, Inc. All Rights Reserved.
+
+#include "FastGeoHLOD.h"
+#include "FastGeoContainer.h"
+#include "FastGeoComponent.h"
+#include "FastGeoStaticMeshComponent.h"
+#include "FastGeoInstancedStaticMeshComponent.h"
+
+#include "Engine/Level.h"
+#include "Engine/World.h"
+#include "PhysicsEngine/BodySetup.h"
+#include "PrimitiveSceneInfo.h"
+#include "PrimitiveSceneProxy.h"
+#include "RenderingThread.h"
+#include "SceneInterface.h"
+#include "WorldPartition/WorldPartitionRuntimeCell.h"
+#include "WorldPartition/HLOD/HLODRuntimeSubsystem.h"
+
+const FFastGeoElementType FFastGeoHLOD::Type(&FFastGeoComponentCluster::Type);
+
+FFastGeoHLOD::FFastGeoHLOD(UFastGeoContainer* InOwner, const FName InName, FFastGeoElementType InType)
+	: Super(InOwner, InName, InType)
+	, bIsVisible(true)
+	, bRequireWarmup(false)
+{
+}
+
+void FFastGeoHLOD::Serialize(FArchive& Ar)
+{
+	Super::Serialize(Ar);
+
+	Ar << bRequireWarmup;
+	Ar << SourceCellGuid;
+}
+
+ULevel* FFastGeoHLOD::GetHLODLevel() const
+{
+	return GetLevel();
+}
+
+FString FFastGeoHLOD::GetHLODNameOrLabel() const
+{
+	return Name;
+}
+
+void FFastGeoHLOD::SetVisibility(bool bInIsVisible)
+{
+	if (bIsVisible != bInIsVisible)
+	{
+		bIsVisible = bInIsVisible;
+		UpdateVisibility();
+	}
+}
+
+bool FFastGeoHLOD::DoesRequireWarmup() const
+{
+	return bRequireWarmup;
+}
+
+TSet<UObject*> FFastGeoHLOD::GetAssetsToWarmup() const
+{
+	TSet<UObject*> AssetsToWarmup;
+
+	ForEachComponent<FFastGeoStaticMeshComponentBase>([&AssetsToWarmup](const FFastGeoStaticMeshComponentBase& SMC)
+	{
+		// Assume ISM HLOD don't need warmup, as they are actually found in the source level
+		if (SMC.IsA<FFastGeoInstancedStaticMeshComponent>())
+		{
+			return;
+		}
+
+		for (int32 iMaterialIndex = 0; iMaterialIndex < SMC.GetNumMaterials(); ++iMaterialIndex)
+		{
+			if (UMaterialInterface* Material = SMC.GetMaterial(iMaterialIndex))
+			{
+				AssetsToWarmup.Add(Material);
+			}
+		}
+
+		if (UStaticMesh* StaticMesh = SMC.GetStaticMesh())
+		{
+			AssetsToWarmup.Add(StaticMesh);
+		}
+	});
+
+	return AssetsToWarmup;
+}
+
+const FGuid& FFastGeoHLOD::GetSourceCellGuid() const
+{
+	// When no source cell guid was set, try resolving it through its associated world partition runtime cell
+	// This is necessary for any HLOD actor part of a level that is instanced multiple times (shared amongst multiple cells)
+	if (!SourceCellGuid.IsValid())
+	{
+		const UWorldPartitionRuntimeCell* Cell = Cast<UWorldPartitionRuntimeCell>(GetLevel()->GetWorldPartitionRuntimeCell());
+		if (Cell && Cell->GetIsHLOD())
+		{
+			const_cast<FFastGeoHLOD*>(this)->SourceCellGuid = Cell->GetSourceCellGuid();
+		}
+	}
+	return SourceCellGuid;
+}
+
+void FFastGeoHLOD::OnRegister()
+{
+	GetLevel()->GetWorld()->GetSubsystem<UWorldPartitionHLODRuntimeSubsystem>()->RegisterHLODObject(this);
+}
+
+void FFastGeoHLOD::OnUnregister()
+{
+	GetLevel()->GetWorld()->GetSubsystem<UWorldPartitionHLODRuntimeSubsystem>()->UnregisterHLODObject(this);
+}
+
+#if WITH_EDITOR
+void FFastGeoHLOD::SetSourceCellGuid(const FGuid& InSourceCellGuid)
+{
+	SourceCellGuid = InSourceCellGuid;
+}
+
+void FFastGeoHLOD::SetRequireWarmup(bool bInRequireWarmup)
+{
+	bRequireWarmup = bInRequireWarmup;
+}
+
+FFastGeoComponent& FFastGeoHLOD::AddComponent(FFastGeoElementType InComponentType)
+{
+	FFastGeoComponent& NewComponent = Super::AddComponent(InComponentType);
+
+	if (FFastGeoPrimitiveComponent* PrimitiveComponent = NewComponent.CastTo<FFastGeoPrimitiveComponent>())
+	{
+		// Always disable collisions on HLODs
+		PrimitiveComponent->SetCollisionEnabled(false);
+	}
+
+	return NewComponent;
+}
+
+void FFastGeoHLOD::PreSave(FObjectPreSaveContext ObjectSaveContext)
+{
+	Super::PreSave(ObjectSaveContext);
+
+	// When cooking, get rid of collision data
+	if (ObjectSaveContext.IsCooking())
+	{
+		ForEachComponent<FFastGeoStaticMeshComponentBase>([this](const FFastGeoStaticMeshComponentBase& StaticMeshComponent)
+		{
+			if (UStaticMesh* StaticMesh = StaticMeshComponent.GetStaticMesh())
+			{
+				// If the HLOD process did create this static mesh
+				if (StaticMesh->GetPackage() == GetOwnerContainer()->GetPackage())
+				{
+					if (UBodySetup* BodySetup = StaticMesh->GetBodySetup())
+					{
+						// To ensure a deterministic cook, save the current GUID and restore it below
+						FGuid PreviousBodySetupGuid = BodySetup->BodySetupGuid;
+						BodySetup->DefaultInstance.SetCollisionProfileName(UCollisionProfile::NoCollision_ProfileName);
+						BodySetup->bNeverNeedsCookedCollisionData = true;
+						BodySetup->bHasCookedCollisionData = false;
+						BodySetup->InvalidatePhysicsData();
+						BodySetup->BodySetupGuid = PreviousBodySetupGuid;
+					}
+				}
+			}
+		});
+	}
+}
+#endif
\ No newline at end of file
diff --git a/Engine/Plugins/Experimental/FastGeoStreaming/Source/FastGeoStreaming/Private/FastGeoHLOD.h b/Engine/Plugins/Experimental/FastGeoStreaming/Source/FastGeoStreaming/Private/FastGeoHLOD.h
new file mode 100644
--- /dev/null
+++ b/Engine/Plugins/Experimental/FastGeoStreaming/Source/FastGeoStreaming/Private/FastGeoHLOD.h
@@ -0,0 +1,55 @@
+// Copyright Epic Games, Inc. All Rights Reserved.
+
+#pragma once
+
+#include "FastGeoComponentCluster.h"
+#include "WorldPartition/HLOD/IWorldPartitionHLODObject.h"
+
+class FFastGeoHLOD : public FFastGeoComponentCluster, public IWorldPartitionHLODObject
+{
+public:
+	typedef FFastGeoComponentCluster Super;
+
+	/** Static type identifier for this element class */
+	static const FFastGeoElementType Type;
+
+	FFastGeoHLOD(UFastGeoContainer* InOwner = nullptr, FName InName = NAME_None, FFastGeoElementType InType = Type);
+	~FFastGeoHLOD() {}
+
+	virtual bool IsVisible() const override { return bIsVisible; }
+
+	// ~Begin FFastGeoComponentCluster interface
+	virtual void OnRegister() override;
+	virtual void OnUnregister() override;
+	virtual void Serialize(FArchive& Ar) override;
+	// ~End FFastGeoComponentCluster interface
+
+	// ~Begin IWorldPartitionHLODObject interface
+	virtual UObject* GetUObject() const override { return nullptr; }
+	virtual ULevel* GetHLODLevel() const override;
+	virtual FString GetHLODNameOrLabel() const override;
+	virtual bool DoesRequireWarmup() const override;
+	virtual TSet<UObject*> GetAssetsToWarmup() const override;
+	virtual void SetVisibility(bool bIsVisible) override;
+	virtual const FGuid& GetSourceCellGuid() const override;
+	// ~End IWorldPartitionHLODObject interface
+
+#if WITH_EDITOR
+	// ~Begin FFastGeoComponentCluster interface
+	virtual FFastGeoComponent& AddComponent(FFastGeoElementType InComponentType) override;
+	virtual void PreSave(FObjectPreSaveContext ObjectSaveContext) override;
+	// ~End FFastGeoComponentCluster interface
+
+	void SetSourceCellGuid(const FGuid& InSourceCellGuid);
+	void SetRequireWarmup(bool bInRequireWarmup);
+#endif
+
+private:
+
+	// Transient Data
+	bool bIsVisible;
+
+	// Persistent Data
+	bool bRequireWarmup;
+	FGuid SourceCellGuid;
+};
\ No newline at end of file
diff --git a/Engine/Plugins/Experimental/FastGeoStreaming/Source/FastGeoStreaming/Private/FastGeoInstancedStaticMeshComponent.cpp b/Engine/Plugins/Experimental/FastGeoStreaming/Source/FastGeoStreaming/Private/FastGeoInstancedStaticMeshComponent.cpp
new file mode 100644
--- /dev/null
+++ b/Engine/Plugins/Experimental/FastGeoStreaming/Source/FastGeoStreaming/Private/FastGeoInstancedStaticMeshComponent.cpp
@@ -0,0 +1,355 @@
+// Copyright Epic Games, Inc. All Rights Reserved.
+
+#include "FastGeoInstancedStaticMeshComponent.h"
+#include "FastGeoComponentCluster.h"
+#include "FastGeoHLOD.h"
+#include "NaniteSceneProxy.h"
+#include "SceneInterface.h"
+#include "PhysicsEngine/BodySetup.h"
+#include "Engine/InstancedStaticMesh.h"
+#include "InstancedStaticMeshComponentHelper.h"
+#include "InstancedStaticMesh/ISMInstanceDataSceneProxy.h"
+#include "NaniteVertexFactory.h"
+#include "FastGeoLog.h"
+
+const FFastGeoElementType FFastGeoInstancedStaticMeshComponent::Type(&FFastGeoStaticMeshComponentBase::Type);
+
+FFastGeoInstancedStaticMeshComponent::FFastGeoInstancedStaticMeshComponent(int32 InComponentIndex, FFastGeoElementType InType)
+	: Super(InComponentIndex, InType)
+{
+}
+
+void FFastGeoInstancedStaticMeshComponent::Serialize(FArchive& Ar)
+{
+	Super::Serialize(Ar);
+
+	// Serialize persistent data from FFastGeoInstancedStaticMeshComponent
+	PerInstanceSMData.BulkSerialize(Ar);
+	Ar << InstancingRandomSeed;
+	PerInstanceSMCustomData.BulkSerialize(Ar);
+	Ar << AdditionalRandomSeeds;
+
+	// Serialize persistent data from FInstancedStaticMeshSceneProxyDesc
+	Ar << SceneProxyDesc.InstanceLODDistanceScale;
+	Ar << SceneProxyDesc.InstanceMinDrawDistance;
+	Ar << SceneProxyDesc.InstanceStartCullDistance;
+	Ar << SceneProxyDesc.InstanceEndCullDistance;
+	FArchive_Serialize_BitfieldBool(Ar, SceneProxyDesc.bUseGpuLodSelection);
+}
+
+FLinearColor FFastGeoInstancedStaticMeshComponent::GetDebugColorImp() const
+{
+	return GetOwnerComponentCluster()->IsA<FFastGeoHLOD>() ? FColor::Red : FColor::Orange;
+}
+
+#if WITH_EDITOR
+void FFastGeoInstancedStaticMeshComponent::InitializeSceneProxyDescFromComponent(UActorComponent* Component)
+{
+	UInstancedStaticMeshComponent* InstancedStaticMeshComponent = CastChecked<UInstancedStaticMeshComponent>(Component);
+	SceneProxyDesc.InitializeFromInstancedStaticMeshComponent(InstancedStaticMeshComponent);
+}
+
+void FFastGeoInstancedStaticMeshComponent::InitializeFromComponent(UActorComponent* Component)
+{
+	Super::InitializeFromComponent(Component);
+
+	UInstancedStaticMeshComponent* ISMComponent = CastChecked<UInstancedStaticMeshComponent>(Component);
+	SceneProxyDesc.bCollisionEnabled = !!SceneProxyDesc.bCollisionEnabled && !ISMComponent->bDisableCollision;
+	AdditionalRandomSeeds = ISMComponent->AdditionalRandomSeeds;
+	PerInstanceSMData = ISMComponent->PerInstanceSMData;
+	PerInstanceSMCustomData = ISMComponent->PerInstanceSMCustomData;
+	InstancingRandomSeed = ISMComponent->InstancingRandomSeed;
+
+	LocalBounds = CalculateBounds(false);
+	WorldBounds = CalculateBounds(true);
+}
+
+UClass* FFastGeoInstancedStaticMeshComponent::GetEditorProxyClass() const
+{
+	return UFastGeoInstancedStaticMeshComponentEditorProxy::StaticClass();
+}
+
+void FFastGeoInstancedStaticMeshComponent::ResetSceneProxyDescUnsupportedProperties()
+{
+	Super::ResetSceneProxyDescUnsupportedProperties();
+
+	SceneProxyDesc.InstanceDataSceneProxy = nullptr;
+	SceneProxyDesc.bHasSelectedInstances = false;
+}
+#endif
+
+void FFastGeoInstancedStaticMeshComponent::ApplyWorldTransform(const FTransform& InTransform)
+{
+	Super::ApplyWorldTransform(InTransform);
+
+	WorldBounds = CalculateBounds(true);
+}
+
+TSharedPtr<FInstanceDataSceneProxy, ESPMode::ThreadSafe>& FFastGeoInstancedStaticMeshComponent::BuildInstanceData()
+{
+	FInstanceSceneDataBuffers InstanceSceneDataBuffers{};
+	FInstanceSceneDataBuffers::FAccessTag AccessTag(PointerHash(this));
+	auto View = InstanceSceneDataBuffers.BeginWriteAccess(AccessTag);
+	
+	// PrimitiveLocalToWorld
+	InstanceSceneDataBuffers.SetPrimitiveLocalToWorld(GetRenderMatrix(), AccessTag);
+	
+	// InstanceLocalBounds
+	const FPrimitiveMaterialPropertyDescriptor PrimitiveMaterialDesc = GetUsedMaterialPropertyDesc(GetScene()->GetFeatureLevel());
+	const float LocalAbsMaxDisplacement = FMath::Max(-PrimitiveMaterialDesc.MinMaxMaterialDisplacement.X, PrimitiveMaterialDesc.MinMaxMaterialDisplacement.Y) + PrimitiveMaterialDesc.MaxWorldPositionOffsetDisplacement;
+	const FVector3f PadExtent = FISMCInstanceDataSceneProxy::GetLocalBoundsPadExtent(View.PrimitiveToRelativeWorld, LocalAbsMaxDisplacement);
+	FRenderBounds InstanceLocalBounds = GetStaticMesh()->GetBounds();
+	InstanceLocalBounds.Min -= PadExtent;
+	InstanceLocalBounds.Max += PadExtent;
+	check(!View.Flags.bHasPerInstanceLocalBounds);
+	View.InstanceLocalBounds.Add(InstanceLocalBounds);
+
+	// LocalToPrimitiveRelativeWorld
+	View.InstanceToPrimitiveRelative.Reserve(PerInstanceSMData.Num());
+	for (const auto& SM : PerInstanceSMData)
+	{
+		FRenderTransform InstanceToPrimitive = SM.Transform;
+		FRenderTransform LocalToPrimitiveRelativeWorld = InstanceToPrimitive * View.PrimitiveToRelativeWorld;
+		LocalToPrimitiveRelativeWorld.Orthogonalize();
+		View.InstanceToPrimitiveRelative.Add(LocalToPrimitiveRelativeWorld);
+	}
+
+	// InstanceCustomData
+	View.InstanceCustomData = PerInstanceSMCustomData;
+	View.NumCustomDataFloats = PerInstanceSMCustomData.Num() && PerInstanceSMData.Num() ? PerInstanceSMCustomData.Num() / PerInstanceSMData.Num() : 0;
+	View.Flags.bHasPerInstanceCustomData = PrimitiveMaterialDesc.bAnyMaterialHasPerInstanceCustomData && View.NumCustomDataFloats != 0;
+	if (!View.Flags.bHasPerInstanceCustomData)
+	{
+		View.NumCustomDataFloats = 0;
+		View.InstanceCustomData.Reset();
+	}
+
+	// InstanceRandomIDs
+	View.Flags.bHasPerInstanceRandom = PrimitiveMaterialDesc.bAnyMaterialHasPerInstanceRandom && (PerInstanceSMData.Num() > 0);
+	if (View.Flags.bHasPerInstanceRandom && InstanceRandomIDs.IsEmpty())
+	{
+		InstanceRandomIDs.SetNumZeroed(PerInstanceSMData.Num());
+		FRandomStream RandomStream = FRandomStream(InstancingRandomSeed);
+		auto AdditionalRandomSeedsIt = AdditionalRandomSeeds.CreateIterator();
+		int32 SeedResetIndex = AdditionalRandomSeedsIt ? AdditionalRandomSeedsIt->StartInstanceIndex : INDEX_NONE;
+		for (int32 Index = 0; Index < InstanceRandomIDs.Num(); ++Index)
+		{
+			// Reset the random stream if necessary
+			if (Index == SeedResetIndex)
+			{
+				RandomStream = FRandomStream(AdditionalRandomSeedsIt->RandomSeed);
+				AdditionalRandomSeedsIt++;
+				SeedResetIndex = AdditionalRandomSeedsIt ? AdditionalRandomSeedsIt->StartInstanceIndex : INDEX_NONE;
+			}
+			InstanceRandomIDs[Index] = RandomStream.GetFraction();
+		}
+	}
+	View.InstanceRandomIDs = InstanceRandomIDs;
+
+	InstanceSceneDataBuffers.EndWriteAccess(AccessTag);
+	InstanceSceneDataBuffers.ValidateData();
+
+	DataProxy = MakeShared<FInstanceDataSceneProxy, ESPMode::ThreadSafe>(MoveTemp(InstanceSceneDataBuffers));
+	return DataProxy;
+}
+
+FPrimitiveSceneProxy* FFastGeoInstancedStaticMeshComponent::CreateSceneProxy(ESceneProxyCreationError* OutError)
+{
+	check(GetWorld());
+	FSceneInterface* Scene = GetScene();
+	UStaticMesh* StaticMesh = GetStaticMesh();
+	check(Scene);
+	check(StaticMesh);
+	check(StaticMesh->GetRenderData());
+	check(StaticMesh->GetRenderData()->IsInitialized());
+	check(!StaticMesh->IsCompiling());
+
+	if (OutError)
+	{
+		*OutError = ESceneProxyCreationError::None;
+	}
+
+#if !WITH_EDITOR && UE_WITH_PSO_PRECACHING
+	const EPSOPrecachePriority PSOPrecachePriority = GetStaticMeshComponentBoostPSOPrecachePriority();
+	if (CheckPSOPrecachingAndBoostPriority(PSOPrecachePriority) && (GetPSOPrecacheProxyCreationStrategy() == EPSOPrecacheProxyCreationStrategy::DelayUntilPSOPrecached))
+	{
+		UE_LOG(LogFastGeoStreaming, Verbose, TEXT("Skipping CreateSceneProxy for FastGeoInstancedStaticMeshComponent of %s (component PSOs are still compiling)"), *GetOwnerComponentCluster()->GetName());
+		if (OutError)
+		{
+			*OutError = ESceneProxyCreationError::WaitingPSOs;
+		}
+		return nullptr;
+	}
+#endif
+
+	InitializeSceneProxyDescDynamicProperties();
+
+	check(SceneProxyDesc.GetScene());
+	check(SceneProxyDesc.GetScene() == Scene);
+	check(SceneProxyDesc.World == Scene->GetWorld());
+	check(SceneProxyDesc.FeatureLevel == Scene->GetFeatureLevel());
+	check(SceneProxyDesc.ComponentId == GetPrimitiveSceneId());
+
+	// Verify that the mesh is valid before using it.
+	const bool bMeshIsValid = PerInstanceSMData.Num() > 0 && StaticMesh && StaticMesh->IsCompiling() == false && StaticMesh->HasValidRenderData();
+	if (!bMeshIsValid)
+	{
+		if (OutError)
+		{
+			*OutError = ESceneProxyCreationError::InvalidMesh;
+		}
+		return nullptr;
+	}
+
+	SceneProxyDesc.InstanceDataSceneProxy = BuildInstanceData();
+	Nanite::FMaterialAudit MaterialAudit;
+	if (SceneProxyDesc.ShouldCreateNaniteProxy(&MaterialAudit))
+	{
+		PrimitiveSceneData.SceneProxy = ::new Nanite::FSceneProxy(MaterialAudit, SceneProxyDesc);
+	}
+	else
+	{
+		PrimitiveSceneData.SceneProxy = ::new FInstancedStaticMeshSceneProxy(SceneProxyDesc, SceneProxyDesc.FeatureLevel);
+	}
+	return PrimitiveSceneData.SceneProxy;
+}
+
+void FFastGeoInstancedStaticMeshComponent::OnAsyncCreatePhysicsState()
+{
+	TRACE_CPUPROFILER_EVENT_SCOPE(FFastGeoInstancedStaticMeshComponent::OnAsyncCreatePhysicsState);
+
+	check(InstanceBodies.Num() == 0);
+
+	FPhysScene* PhysScene = GetWorld()->GetPhysicsScene();
+	if (!PhysScene)
+	{
+		return;
+	}
+
+	// Create all the bodies.
+	CreateAllInstanceBodies();
+
+	FFastGeoComponent::OnAsyncCreatePhysicsState();
+}
+
+void FFastGeoInstancedStaticMeshComponent::OnAsyncDestroyPhysicsStateBegin_GameThread()
+{
+	TRACE_CPUPROFILER_EVENT_SCOPE(FFastGeoInstancedStaticMeshComponent::OnAsyncDestroyPhysicsStateBegin_GameThread);
+
+	FFastGeoComponent::OnAsyncDestroyPhysicsStateBegin_GameThread();
+
+	// Move InstanceBodies in AsyncDestroyPhysicsStatePayload
+	check(AsyncDestroyPhysicsStatePayload.IsEmpty());
+	AsyncDestroyPhysicsStatePayload = MoveTemp(InstanceBodies);
+}
+
+void FFastGeoInstancedStaticMeshComponent::OnAsyncDestroyPhysicsState()
+{
+	TRACE_CPUPROFILER_EVENT_SCOPE(FFastGeoInstancedStaticMeshComponent::OnAsyncDestroyPhysicsState);
+
+	FFastGeoComponent::OnAsyncDestroyPhysicsState();
+
+	check(InstanceBodies.IsEmpty());
+	for (FBodyInstance*& Instance : AsyncDestroyPhysicsStatePayload)
+	{
+		if (Instance)
+		{
+			Instance->TermBody();
+			delete Instance;
+		}
+	}
+	AsyncDestroyPhysicsStatePayload.Empty();
+}
+
+void FFastGeoInstancedStaticMeshComponent::CreateAllInstanceBodies()
+{
+	TRACE_CPUPROFILER_EVENT_SCOPE(FFastGeoInstancedStaticMeshComponent::CreateAllInstanceBodies);
+	
+	const int32 NumBodies = PerInstanceSMData.Num();
+	check(InstanceBodies.Num() == 0);
+	check(SceneProxyDesc.Mobility != EComponentMobility::Movable);
+
+	if (UBodySetup* BodySetup = GetBodySetup())
+	{
+		FPhysScene* PhysScene = GetWorld()->GetPhysicsScene();
+
+		if (!BodyInstance.GetOverrideWalkableSlopeOnInstance())
+		{
+			BodyInstance.SetWalkableSlopeOverride(BodySetup->WalkableSlopeOverride, false);
+		}
+
+		InstanceBodies.SetNumUninitialized(NumBodies);
+
+		// Sanitized array does not contain any nulls
+		TArray<FBodyInstance*> InstanceBodiesSanitized;
+		InstanceBodiesSanitized.Reserve(NumBodies);
+
+		TArray<FTransform> Transforms;
+		Transforms.Reserve(NumBodies);
+		for (int32 i = 0; i < NumBodies; ++i)
+		{
+			const FTransform InstanceTM = FTransform(PerInstanceSMData[i].Transform) * WorldTransform;
+			if (InstanceTM.GetScale3D().IsNearlyZero())
+			{
+				InstanceBodies[i] = nullptr;
+			}
+			else
+			{
+				FBodyInstance* Instance = new FBodyInstance;
+
+				InstanceBodiesSanitized.Add(Instance);
+				InstanceBodies[i] = Instance;
+				Instance->CopyBodyInstancePropertiesFrom(&BodyInstance);
+				Instance->InstanceBodyIndex = i;
+				Instance->bAutoWeld = false;
+				Instance->bSimulatePhysics = false;
+				Transforms.Add(InstanceTM);
+			}
+		}
+
+		if (InstanceBodiesSanitized.Num() > 0)
+		{
+			FBodyInstance::InitStaticBodies(MoveTemp(InstanceBodiesSanitized), MoveTemp(Transforms), BodySetup, nullptr, GetWorld()->GetPhysicsScene(), this);
+		}
+	}
+	else
+	{
+		// In case we get into some bad state where the BodySetup is invalid but bPhysicsStateCreated is true,
+		// issue a warning and add nullptrs to InstanceBodies.
+		UE_LOG(LogFastGeoStreaming, Warning, TEXT("Instance Static Mesh Component unable to create InstanceBodies!"));
+		InstanceBodies.AddZeroed(NumBodies);
+	}
+}
+
+FBoxSphereBounds FFastGeoInstancedStaticMeshComponent::CalculateBounds(bool bWorldSpace)
+{
+	if (GetStaticMesh() && PerInstanceSMData.Num() > 0)
+	{
+		const FBox InstanceBounds = GetStaticMesh()->GetBounds().GetBox();
+		const FMatrix ComponentTransformMatrix = WorldTransform.ToMatrixWithScale();
+		if (InstanceBounds.IsValid)
+		{
+			FBoxSphereBounds::Builder BoundsBuilder;
+			for (int32 InstanceIndex = 0; InstanceIndex < PerInstanceSMData.Num(); InstanceIndex++)
+			{
+				if (bWorldSpace)
+				{
+					BoundsBuilder += InstanceBounds.TransformBy(PerInstanceSMData[InstanceIndex].Transform * ComponentTransformMatrix);
+				}
+				else
+				{
+					BoundsBuilder += InstanceBounds.TransformBy(PerInstanceSMData[InstanceIndex].Transform);
+				}
+			}
+			return BoundsBuilder;
+		}
+	}
+	return FBoxSphereBounds(FVector::ZeroVector, FVector::ZeroVector, 0.f);
+};
+
+#if !WITH_EDITOR && UE_WITH_PSO_PRECACHING
+void FFastGeoInstancedStaticMeshComponent::CollectPSOPrecacheData(const FPSOPrecacheParams& BasePrecachePSOParams, FMaterialInterfacePSOPrecacheParamsList& OutParams)
+{
+	return FInstancedStaticMeshComponentHelper::CollectPSOPrecacheData(*this, BasePrecachePSOParams, OutParams);
+}
+#endif
\ No newline at end of file
diff --git a/Engine/Plugins/Experimental/FastGeoStreaming/Source/FastGeoStreaming/Private/FastGeoInstancedStaticMeshComponent.h b/Engine/Plugins/Experimental/FastGeoStreaming/Source/FastGeoStreaming/Private/FastGeoInstancedStaticMeshComponent.h
new file mode 100644
--- /dev/null
+++ b/Engine/Plugins/Experimental/FastGeoStreaming/Source/FastGeoStreaming/Private/FastGeoInstancedStaticMeshComponent.h
@@ -0,0 +1,93 @@
+// Copyright Epic Games, Inc. All Rights Reserved.
+
+#pragma once
+
+#include "FastGeoStaticMeshComponent.h"
+#include "InstancedStaticMeshSceneProxyDesc.h"
+#include "FastGeoInstancedStaticMeshComponent.generated.h"
+
+class FASTGEOSTREAMING_API FFastGeoInstancedStaticMeshComponent : public FFastGeoStaticMeshComponentBase
+{
+public:
+	typedef FFastGeoStaticMeshComponentBase Super;
+
+	/** Static type identifier for this element class */
+	static const FFastGeoElementType Type;
+
+	FFastGeoInstancedStaticMeshComponent(int32 InComponentIndex = INDEX_NONE, FFastGeoElementType InType = Type);
+	virtual ~FFastGeoInstancedStaticMeshComponent() = default;
+
+	//~ Begin FFastGeoComponent Interface
+	virtual void Serialize(FArchive& Ar) override;
+	virtual void OnAsyncCreatePhysicsState() override;
+	virtual void OnAsyncDestroyPhysicsState() override;
+	virtual void OnAsyncDestroyPhysicsStateBegin_GameThread() override;
+#if WITH_EDITOR
+	virtual void InitializeFromComponent(UActorComponent* Component) override;
+	virtual UClass* GetEditorProxyClass() const override;
+#endif
+	//~ End FFastGeoComponent Interface
+
+protected:
+	//~ Being FFastGeoComponent interface
+	virtual FLinearColor GetDebugColorImp() const override;
+	//~ End FFastGeoComponent interface
+
+	//~ Begin FFastGeoPrimitiveComponent interface
+	virtual FPrimitiveSceneProxyDesc& GetSceneProxyDesc() override { return SceneProxyDesc; }
+	virtual const FPrimitiveSceneProxyDesc& GetSceneProxyDesc() const override { return SceneProxyDesc; }
+#if WITH_EDITOR
+	virtual void InitializeSceneProxyDescFromComponent(UActorComponent* Component) override;
+	virtual void ResetSceneProxyDescUnsupportedProperties() override;
+#endif
+	virtual void ApplyWorldTransform(const FTransform& InTransform) override;
+	virtual FPrimitiveSceneProxy* CreateSceneProxy(ESceneProxyCreationError* OutError = nullptr) override;
+	//~ End FFastGeoPrimitiveComponent interface
+
+	//~ Begin FFastGeoStaticMeshComponentBase interface
+	virtual FStaticMeshSceneProxyDesc& GetStaticMeshSceneProxyDesc() override { return SceneProxyDesc; }
+	virtual const FStaticMeshSceneProxyDesc& GetStaticMeshSceneProxyDesc() const override { return SceneProxyDesc; }
+#if !WITH_EDITOR && UE_WITH_PSO_PRECACHING
+	virtual void CollectPSOPrecacheData(const FPSOPrecacheParams& BasePrecachePSOParams, FMaterialInterfacePSOPrecacheParamsList& OutParams) override;
+#endif
+	//~ End FFastGeoStaticMeshComponentBase interface
+
+	TSharedPtr<FInstanceDataSceneProxy, ESPMode::ThreadSafe>& BuildInstanceData();
+
+private:
+	void CreateAllInstanceBodies();
+	FBoxSphereBounds CalculateBounds(bool bWorldSpace);
+
+private:
+
+	// Persistent Data
+	TArray<FInstancedStaticMeshInstanceData> PerInstanceSMData;
+	int32 InstancingRandomSeed = 0;
+	TArray<float> PerInstanceSMCustomData;
+	TArray<FInstancedStaticMeshRandomSeed> AdditionalRandomSeeds;
+	float InstanceLODDistanceScale = 1.f;
+	int32 InstanceStartCullDistance = 0;
+	int32 InstanceEndCullDistance = 0;
+	uint8 bUseGpuLodSelection : 1 = false;
+	FInstancedStaticMeshSceneProxyDesc SceneProxyDesc{};
+
+	// Transient data
+	TSharedPtr<FInstanceDataSceneProxy, ESPMode::ThreadSafe> DataProxy{};
+	TArray<float> InstanceRandomIDs;
+	
+	/** Physics representation of the instance bodies. */
+	TArray<FBodyInstance*> InstanceBodies;
+
+	/** Payload used by asynchronous destruction of physics state (see OnAsyncDestroyPhysicsState). */
+	TArray<FBodyInstance*> AsyncDestroyPhysicsStatePayload;
+};
+
+UCLASS()
+class UFastGeoInstancedStaticMeshComponentEditorProxy : public UFastGeoStaticMeshComponentEditorProxy
+{
+	GENERATED_BODY()
+
+#if WITH_EDITOR
+	typedef FFastGeoInstancedStaticMeshComponent ComponentType;
+#endif
+};
diff --git a/Engine/Plugins/Experimental/FastGeoStreaming/Source/FastGeoStreaming/Private/FastGeoLog.cpp b/Engine/Plugins/Experimental/FastGeoStreaming/Source/FastGeoStreaming/Private/FastGeoLog.cpp
new file mode 100644
--- /dev/null
+++ b/Engine/Plugins/Experimental/FastGeoStreaming/Source/FastGeoStreaming/Private/FastGeoLog.cpp
@@ -0,0 +1,5 @@
+// Copyright Epic Games, Inc. All Rights Reserved.
+
+#include "FastGeoLog.h"
+
+DEFINE_LOG_CATEGORY(LogFastGeoStreaming);
\ No newline at end of file
diff --git a/Engine/Plugins/Experimental/FastGeoStreaming/Source/FastGeoStreaming/Private/FastGeoLog.h b/Engine/Plugins/Experimental/FastGeoStreaming/Source/FastGeoStreaming/Private/FastGeoLog.h
new file mode 100644
--- /dev/null
+++ b/Engine/Plugins/Experimental/FastGeoStreaming/Source/FastGeoStreaming/Private/FastGeoLog.h
@@ -0,0 +1,7 @@
+// Copyright Epic Games, Inc. All Rights Reserved.
+
+#pragma once
+
+#include "CoreMinimal.h"
+
+FASTGEOSTREAMING_API DECLARE_LOG_CATEGORY_EXTERN(LogFastGeoStreaming, Log, All);
\ No newline at end of file
diff --git a/Engine/Plugins/Experimental/FastGeoStreaming/Source/FastGeoStreaming/Private/FastGeoPrimitiveComponent.cpp b/Engine/Plugins/Experimental/FastGeoStreaming/Source/FastGeoStreaming/Private/FastGeoPrimitiveComponent.cpp
new file mode 100644
--- /dev/null
+++ b/Engine/Plugins/Experimental/FastGeoStreaming/Source/FastGeoStreaming/Private/FastGeoPrimitiveComponent.cpp
@@ -0,0 +1,836 @@
+// Copyright Epic Games, Inc. All Rights Reserved.
+
+#include "FastGeoPrimitiveComponent.h"
+#include "FastGeoContainer.h"
+#include "FastGeoComponentCluster.h"
+#include "FastGeoInstancedStaticMeshComponent.h"
+#include "FastGeoWorldSubsystem.h"
+#include "FastGeoComponentCluster.h"
+#include "FastGeoWeakElement.h"
+#include "PrimitiveSceneDesc.h"
+#include "SceneInterface.h"
+#include "WorldPartition/WorldPartition.h"
+#include "PhysicsEngine/BodyInstance.h"
+#include "PhysicsEngine/BodySetup.h"
+#include "PSOPrecache.h"
+#include "PSOPrecacheMaterial.h"
+#include "FastGeoLog.h"
+
+#if !UE_BUILD_SHIPPING
+namespace FastGeo
+{
+	static int32 GShowFastGeo = 1;
+	FAutoConsoleCommand GShowFastGeoCommand(
+		TEXT("FastGeo.Show"),
+		TEXT("Turn on/off rendering of FastGeo."),
+		FConsoleCommandWithArgsDelegate::CreateLambda([](const TArray<FString>& Args)
+		{
+			GShowFastGeo = (Args.Num() != 1) || (Args[0] != TEXT("0"));
+			for (const FWorldContext& Context : GEngine->GetWorldContexts())
+			{
+				UWorld* World = Context.World();
+				if (World && World->IsGameWorld())
+				{
+					for (ULevel* Level : World->GetLevels())
+					{
+						if (UFastGeoContainer* FastGeo = Level->GetAssetUserData<UFastGeoContainer>())
+						{
+							FastGeo->ForEachComponentCluster([](FFastGeoComponentCluster& ComponentCluster)
+							{
+								ComponentCluster.UpdateVisibility();
+							});
+						}
+					}
+				}
+			}
+		})
+	);
+}
+#endif
+
+const FFastGeoElementType FFastGeoPrimitiveComponent::Type(&FFastGeoComponent::Type);
+
+FFastGeoPrimitiveComponent::FFastGeoPrimitiveComponent(int32 InComponentIndex, FFastGeoElementType InType)
+	: Super(InComponentIndex, InType)
+	, LocalBounds(ForceInit)
+	, WorldBounds(ForceInit)
+{
+}
+
+void FFastGeoPrimitiveComponent::Serialize(FArchive& Ar)
+{
+	Super::Serialize(Ar);
+
+	// Serialize persistent data from FFastGeoPrimitiveComponent
+	Ar << LocalTransform;
+	Ar << WorldTransform;
+	Ar << LocalBounds;
+	Ar << WorldBounds;
+	FArchive_Serialize_BitfieldBool(Ar, bIsVisible);
+	FArchive_Serialize_BitfieldBool(Ar, bStaticWhenNotMoveable);
+	Ar << CustomPrimitiveData.Data;
+	Ar << DetailMode;
+	Ar << RuntimeVirtualTextures;
+	FBodyInstance::StaticStruct()->SerializeItem(Ar, &BodyInstance, nullptr);
+
+	// Serialize persistent data from FPrimitiveSceneProxyDesc
+	FPrimitiveSceneProxyDesc& SceneProxyDesc = GetSceneProxyDesc();
+	FArchive_Serialize_BitfieldBool(Ar, SceneProxyDesc.CastShadow);
+	FArchive_Serialize_BitfieldBool(Ar, SceneProxyDesc.bReceivesDecals);
+	FArchive_Serialize_BitfieldBool(Ar, SceneProxyDesc.bOnlyOwnerSee);
+	FArchive_Serialize_BitfieldBool(Ar, SceneProxyDesc.bOwnerNoSee);
+	FArchive_Serialize_BitfieldBool(Ar, SceneProxyDesc.bUseViewOwnerDepthPriorityGroup);
+	FArchive_Serialize_BitfieldBool(Ar, SceneProxyDesc.bVisibleInReflectionCaptures);
+	FArchive_Serialize_BitfieldBool(Ar, SceneProxyDesc.bVisibleInRealTimeSkyCaptures);
+	FArchive_Serialize_BitfieldBool(Ar, SceneProxyDesc.bVisibleInRayTracing);
+	FArchive_Serialize_BitfieldBool(Ar, SceneProxyDesc.bRenderInDepthPass);
+	FArchive_Serialize_BitfieldBool(Ar, SceneProxyDesc.bRenderInMainPass);
+	FArchive_Serialize_BitfieldBool(Ar, SceneProxyDesc.bTreatAsBackgroundForOcclusion);
+	FArchive_Serialize_BitfieldBool(Ar, SceneProxyDesc.bCastDynamicShadow);
+	FArchive_Serialize_BitfieldBool(Ar, SceneProxyDesc.bCastStaticShadow);
+	FArchive_Serialize_BitfieldBool(Ar, SceneProxyDesc.bEmissiveLightSource);
+	FArchive_Serialize_BitfieldBool(Ar, SceneProxyDesc.bAffectDynamicIndirectLighting);
+	FArchive_Serialize_BitfieldBool(Ar, SceneProxyDesc.bAffectIndirectLightingWhileHidden);
+	FArchive_Serialize_BitfieldBool(Ar, SceneProxyDesc.bAffectDistanceFieldLighting);
+	FArchive_Serialize_BitfieldBool(Ar, SceneProxyDesc.bCastVolumetricTranslucentShadow);
+	FArchive_Serialize_BitfieldBool(Ar, SceneProxyDesc.bCastContactShadow);
+	FArchive_Serialize_BitfieldBool(Ar, SceneProxyDesc.bCastHiddenShadow);
+	FArchive_Serialize_BitfieldBool(Ar, SceneProxyDesc.bCastShadowAsTwoSided);
+	FArchive_Serialize_BitfieldBool(Ar, SceneProxyDesc.bSelfShadowOnly);
+	FArchive_Serialize_BitfieldBool(Ar, SceneProxyDesc.bCastInsetShadow);
+	FArchive_Serialize_BitfieldBool(Ar, SceneProxyDesc.bCastCinematicShadow);
+	FArchive_Serialize_BitfieldBool(Ar, SceneProxyDesc.bCastFarShadow);
+	FArchive_Serialize_BitfieldBool(Ar, SceneProxyDesc.bLightAttachmentsAsGroup);
+	FArchive_Serialize_BitfieldBool(Ar, SceneProxyDesc.bSingleSampleShadowFromStationaryLights);
+	FArchive_Serialize_BitfieldBool(Ar, SceneProxyDesc.bUseAsOccluder);
+	FArchive_Serialize_BitfieldBool(Ar, SceneProxyDesc.bHasPerInstanceHitProxies);
+	FArchive_Serialize_BitfieldBool(Ar, SceneProxyDesc.bReceiveMobileCSMShadows);
+	FArchive_Serialize_BitfieldBool(Ar, SceneProxyDesc.bRenderCustomDepth);
+	FArchive_Serialize_BitfieldBool(Ar, SceneProxyDesc.bVisibleInSceneCaptureOnly);
+	FArchive_Serialize_BitfieldBool(Ar, SceneProxyDesc.bHiddenInSceneCapture);
+	FArchive_Serialize_BitfieldBool(Ar, SceneProxyDesc.bRayTracingFarField);
+	FArchive_Serialize_BitfieldBool(Ar, SceneProxyDesc.bHoldout);
+	FArchive_Serialize_BitfieldBool(Ar, SceneProxyDesc.bIsFirstPerson);
+	FArchive_Serialize_BitfieldBool(Ar, SceneProxyDesc.bIsFirstPersonWorldSpaceRepresentation);
+	FArchive_Serialize_BitfieldBool(Ar, SceneProxyDesc.bCollisionEnabled);
+	FArchive_Serialize_BitfieldBool(Ar, SceneProxyDesc.bSupportsWorldPositionOffsetVelocity);
+	FArchive_Serialize_BitfieldBool(Ar, SceneProxyDesc.bIsInstancedStaticMesh);
+	FArchive_Serialize_BitfieldBool(Ar, SceneProxyDesc.bHasStaticLighting);
+	FArchive_Serialize_BitfieldBool(Ar, SceneProxyDesc.bHasValidSettingsForStaticLighting);
+	FArchive_Serialize_BitfieldBool(Ar, SceneProxyDesc.bIsPrecomputedLightingValid);
+	FArchive_Serialize_BitfieldBool(Ar, SceneProxyDesc.bShadowIndirectOnly);
+	Ar << SceneProxyDesc.Mobility;
+	Ar << SceneProxyDesc.TranslucencySortPriority;
+	Ar << SceneProxyDesc.TranslucencySortDistanceOffset;
+	Ar << SceneProxyDesc.LightmapType;
+	Ar << SceneProxyDesc.ViewOwnerDepthPriorityGroup;
+	Ar << SceneProxyDesc.CustomDepthStencilValue;
+	Ar << SceneProxyDesc.CustomDepthStencilWriteMask;
+	FArchive_Serialize_BitfieldBool(Ar, SceneProxyDesc.LightingChannels.bChannel0);
+	FArchive_Serialize_BitfieldBool(Ar, SceneProxyDesc.LightingChannels.bChannel1);
+	FArchive_Serialize_BitfieldBool(Ar, SceneProxyDesc.LightingChannels.bChannel2);
+	Ar << SceneProxyDesc.RayTracingGroupCullingPriority;
+	Ar << SceneProxyDesc.IndirectLightingCacheQuality;
+	Ar << SceneProxyDesc.ShadowCacheInvalidationBehavior;
+	Ar << SceneProxyDesc.DepthPriorityGroup;
+	Ar << SceneProxyDesc.VirtualTextureLodBias;
+	Ar << SceneProxyDesc.VirtualTextureCullMips;
+	Ar << SceneProxyDesc.VirtualTextureMinCoverage;
+	Ar << SceneProxyDesc.VisibilityId;
+	Ar << SceneProxyDesc.CachedMaxDrawDistance;
+	Ar << SceneProxyDesc.MinDrawDistance;
+	Ar << SceneProxyDesc.BoundsScale;
+	Ar << SceneProxyDesc.RayTracingGroupId;
+	Ar << SceneProxyDesc.VirtualTextureRenderPassType;
+	Ar << SceneProxyDesc.VirtualTextureMainPassMaxDrawDistance;
+}
+
+void FFastGeoPrimitiveComponent::ApplyWorldTransform(const FTransform& InTransform)
+{
+	check(!GetOwnerContainer()->IsRegistered());
+	WorldTransform = LocalTransform * InTransform;
+}
+
+#if WITH_EDITOR
+
+void FFastGeoPrimitiveComponent::InitializeFromComponent(UActorComponent* Component)
+{
+	Super::InitializeFromComponent(Component);
+
+	// Initialize properties not handled by InitializeFromPrimitiveComponent
+	UPrimitiveComponent* PrimitiveComponent = CastChecked<UPrimitiveComponent>(Component);
+	PrimitiveComponent->UpdateComponentToWorld();
+	LocalTransform = PrimitiveComponent->GetComponentToWorld();
+	WorldTransform = LocalTransform;
+	bIsVisible = PrimitiveComponent->IsVisible();
+	bStaticWhenNotMoveable = PrimitiveComponent->GetStaticWhenNotMoveable();
+	CustomPrimitiveData = PrimitiveComponent->GetCustomPrimitiveData();
+	DetailMode = PrimitiveComponent->DetailMode;
+	BodyInstance.CopyBodyInstancePropertiesFrom(&PrimitiveComponent->BodyInstance);
+	RuntimeVirtualTextures = PrimitiveComponent->GetRuntimeVirtualTextures();
+
+	// Initialize SceneProxyDesc from component
+	InitializeSceneProxyDescFromComponent(Component);
+
+	// Reset some values that are not used in FastGeo
+	ResetSceneProxyDescUnsupportedProperties();
+}
+
+UClass* FFastGeoPrimitiveComponent::GetEditorProxyClass() const
+{
+	return UFastGeoPrimitiveComponentEditorProxy::StaticClass();
+}
+
+void FFastGeoPrimitiveComponent::ResetSceneProxyDescUnsupportedProperties()
+{
+	// Unsupported properties
+	FPrimitiveSceneProxyDesc& SceneProxyDesc = GetSceneProxyDesc();
+	SceneProxyDesc.bLevelInstanceEditingState = false;
+	SceneProxyDesc.bSelectable = false;
+	SceneProxyDesc.bUseEditorCompositing = false;
+	SceneProxyDesc.bIsBeingMovedByEditor = false;
+	SceneProxyDesc.bSelected = false;
+	SceneProxyDesc.bIndividuallySelected = false;
+	SceneProxyDesc.bShouldRenderSelected = false;
+	SceneProxyDesc.bWantsEditorEffects = false;
+	SceneProxyDesc.bIsHidden = false;
+	SceneProxyDesc.bIsHiddenEd = false;
+	SceneProxyDesc.bIsOwnerEditorOnly = false;
+	SceneProxyDesc.bIsOwnedByFoliage = false;
+	SceneProxyDesc.HiddenEditorViews = 0;
+	SceneProxyDesc.OverlayColor = FColor(EForceInit::ForceInitToZero);
+	SceneProxyDesc.Component = nullptr;
+
+	// Properties that will be initialized by InitializeSceneProxyDescDynamicProperties
+	SceneProxyDesc.ComponentId = FPrimitiveComponentId();
+	SceneProxyDesc.StatId = TStatId();
+	SceneProxyDesc.Owner = nullptr;
+	SceneProxyDesc.World = nullptr;
+	SceneProxyDesc.CustomPrimitiveData = nullptr;
+	SceneProxyDesc.Scene = nullptr;
+	SceneProxyDesc.PrimitiveComponentInterface = nullptr;
+	SceneProxyDesc.FeatureLevel = ERHIFeatureLevel::Num;
+	SceneProxyDesc.RuntimeVirtualTextures = TArrayView<URuntimeVirtualTexture*>();
+	SceneProxyDesc.bIsVisible = false;
+	SceneProxyDesc.bShouldRenderProxyFallbackToDefaultMaterial = false;
+#if MESH_DRAW_COMMAND_STATS
+	SceneProxyDesc.MeshDrawCommandStatsCategory = NAME_None;
+#endif
+}
+#endif
+
+bool FFastGeoPrimitiveComponent::IsVisible() const
+{
+	return GetSceneProxyDesc().bIsVisible;
+}
+
+void FFastGeoPrimitiveComponent::UpdateVisibility()
+{
+	// Update SceneProxyDesc.bIsVisible as it's dependant on component and component cluster visibility
+	FPrimitiveSceneProxyDesc& SceneProxyDesc = GetSceneProxyDesc();
+	SceneProxyDesc.bIsVisible = bIsVisible && GetOwnerComponentCluster()->IsVisible();
+#if !UE_BUILD_SHIPPING
+	SceneProxyDesc.bIsVisible = SceneProxyDesc.bIsVisible && FastGeo::GShowFastGeo;
+#endif
+}
+
+void FFastGeoPrimitiveComponent::InitializeSceneProxyDescDynamicProperties()
+{
+	check(GetWorld());
+	check(GetScene());
+
+#if WITH_EDITOR
+	ResetSceneProxyDescUnsupportedProperties();
+#endif
+
+	// Initialize non-serialized properties
+	FPrimitiveSceneProxyDesc& SceneProxyDesc = GetSceneProxyDesc();
+	SceneProxyDesc.ComponentId = GetPrimitiveSceneId();
+	UObject const* AdditionalStatObjectPtr = AdditionalStatObject();
+	SceneProxyDesc.StatId = AdditionalStatObjectPtr ? AdditionalStatObjectPtr->GetStatID(true) : TStatId();
+	SceneProxyDesc.Owner = GetOwnerContainer();
+#if !WITH_STATE_STREAM
+	SceneProxyDesc.World = GetWorld();
+#endif
+	SceneProxyDesc.CustomPrimitiveData = &CustomPrimitiveData;
+	SceneProxyDesc.Scene = GetScene();
+#if WITH_EDITOR
+	SceneProxyDesc.PrimitiveComponentInterface = GetEditorProxy<UFastGeoPrimitiveComponentEditorProxy>()->GetPrimitiveComponentInterface();
+#endif
+	SceneProxyDesc.FeatureLevel = SceneProxyDesc.Scene->GetFeatureLevel();
+	TArray<URuntimeVirtualTexture*> const& VirtualTextures = GetRuntimeVirtualTextures();
+	SceneProxyDesc.RuntimeVirtualTextures = MakeArrayView(const_cast<URuntimeVirtualTexture**>(VirtualTextures.GetData()), VirtualTextures.Num());
+	SceneProxyDesc.bShouldRenderProxyFallbackToDefaultMaterial = ShouldRenderProxyFallbackToDefaultMaterial();
+#if MESH_DRAW_COMMAND_STATS
+	static FName NAME_FastGeoPrimitiveComponent(TEXT("FastGeoPrimitiveComponent"));
+	SceneProxyDesc.MeshDrawCommandStatsCategory = NAME_FastGeoPrimitiveComponent;
+#endif
+	check(SceneProxyDesc.bIsInstancedStaticMesh == this->IsA<FFastGeoInstancedStaticMeshComponent>());
+	UpdateVisibility();
+}
+
+FPrimitiveSceneDesc FFastGeoPrimitiveComponent::BuildSceneDesc()
+{
+	check(GetSceneProxy());
+
+	FPrimitiveSceneProxyDesc& SceneProxyDesc = GetSceneProxyDesc();
+	FPrimitiveSceneDesc SceneDesc;
+	SceneDesc.SceneProxy = GetSceneProxy();
+	SceneDesc.ProxyDesc = &SceneProxyDesc;
+	SceneDesc.PrimitiveSceneData = &PrimitiveSceneData;
+	SceneDesc.RenderMatrix = GetRenderMatrix();
+	SceneDesc.AttachmentRootPosition = GetTransform().GetTranslation();
+	SceneDesc.LocalBounds = LocalBounds;
+	SceneDesc.Bounds = GetBounds();
+	SceneDesc.Mobility = SceneProxyDesc.Mobility;
+
+	return SceneDesc;
+}
+
+void FFastGeoPrimitiveComponent::CreateRenderState(FRegisterComponentContext* Context)
+{
+	ProxyState = ProxyCreationState::Creating;
+	bRenderStateDirty = false;
+
+#if WITH_EDITOR
+	GetEditorProxy<EditorProxyType>()->NotifyRenderStateChanged();
+#endif
+
+	FSceneInterface* Scene = GetScene();
+	check(Scene);
+
+	ESceneProxyCreationError Error;
+	if (FPrimitiveSceneProxy* SceneProxy = CreateSceneProxy(&Error))
+	{
+		SceneProxy->SetPrimitiveColor(GetDebugColor());
+		check(GetSceneProxy());
+		FPrimitiveSceneDesc Desc = BuildSceneDesc();
+		Scene->AddPrimitive(&Desc);
+
+		ProxyState = ProxyCreationState::Created;
+	}
+	else if (Error == ESceneProxyCreationError::WaitingPSOs)
+	{
+		ProxyState = ProxyCreationState::Delayed;
+	}
+	else
+	{
+		ProxyState = ProxyCreationState::None;
+	}
+}
+
+FFastGeoPrimitiveComponent::FFastGeoDestroyRenderStateContext::FFastGeoDestroyRenderStateContext(FSceneInterface* InScene)
+	: Scene(InScene)
+{
+}
+
+FFastGeoPrimitiveComponent::FFastGeoDestroyRenderStateContext::~FFastGeoDestroyRenderStateContext()
+{
+	if (HasPendingWork())
+	{
+		Scene->BatchRemovePrimitives(MoveTemp(PrimitiveSceneProxies));
+	}
+}
+
+bool FFastGeoPrimitiveComponent::FFastGeoDestroyRenderStateContext::HasPendingWork() const
+{
+	return !PrimitiveSceneProxies.IsEmpty();
+}
+
+void FFastGeoPrimitiveComponent::FFastGeoDestroyRenderStateContext::DestroyProxy(FFastGeoDestroyRenderStateContext* InContext, FPrimitiveSceneProxy* InPrimitiveSceneProxy)
+{
+	if (InContext)
+	{
+		InContext->PrimitiveSceneProxies.Add(InPrimitiveSceneProxy);
+	}
+	else
+	{
+		InPrimitiveSceneProxy->GetScene().BatchRemovePrimitives({ InPrimitiveSceneProxy });
+	}
+}
+
+void FFastGeoPrimitiveComponent::DestroyRenderState(FFastGeoDestroyRenderStateContext* Context)
+{
+	if (GetSceneProxy())
+	{
+		check(ProxyState == ProxyCreationState::Created);
+
+		FFastGeoDestroyRenderStateContext::DestroyProxy(Context, GetSceneProxy());
+
+		PrimitiveSceneData.SceneProxy = nullptr;
+
+		ProxyState = ProxyCreationState::Pending;
+
+		bRenderStateDirty = false;
+
+#if WITH_EDITOR
+		GetEditorProxy<EditorProxyType>()->NotifyRenderStateChanged();
+#endif
+	}
+}
+
+void FFastGeoPrimitiveComponent::OnAsyncCreatePhysicsState()
+{
+	TRACE_CPUPROFILER_EVENT_SCOPE(FFastGeoPrimitiveComponent::OnAsyncCreatePhysicsState);
+
+	Super::OnAsyncCreatePhysicsState();
+
+	// if we have a scene, we don't want to disable all physics and we have no bodyinstance already
+	if (!BodyInstance.IsValidBodyInstance())
+	{
+		if (UBodySetup* BodySetup = GetBodySetup())
+		{
+			// Create new BodyInstance at given location.
+			FTransform BodyTransform = WorldTransform;
+
+			// Here we make sure we don't have zero scale. This still results in a body being made and placed in
+			// world (very small) but is consistent with a body scaled to zero.
+			const FVector BodyScale = BodyTransform.GetScale3D();
+			if (BodyScale.IsNearlyZero())
+			{
+				BodyTransform.SetScale3D(FVector(UE_KINDA_SMALL_NUMBER));
+			}
+#if !(UE_BUILD_SHIPPING || UE_BUILD_TEST)
+			if ((BodyInstance.GetCollisionEnabled() != ECollisionEnabled::NoCollision) && (FMath::IsNearlyZero(BodyScale.X) || FMath::IsNearlyZero(BodyScale.Y) || FMath::IsNearlyZero(BodyScale.Z)))
+			{
+				UE_LOG(LogFastGeoStreaming, Warning, TEXT("Scale for FastGeoPrimitiveComponent has a component set to zero, which will result in a bad body instance. Scale:%s"), *BodyScale.ToString());
+
+				// User warning has been output - fix up the scale to be valid for physics
+				BodyTransform.SetScale3D(FVector(
+					FMath::IsNearlyZero(BodyScale.X) ? UE_KINDA_SMALL_NUMBER : BodyScale.X,
+					FMath::IsNearlyZero(BodyScale.Y) ? UE_KINDA_SMALL_NUMBER : BodyScale.Y,
+					FMath::IsNearlyZero(BodyScale.Z) ? UE_KINDA_SMALL_NUMBER : BodyScale.Z
+				));
+			}
+#endif
+
+			// Create the body.
+			BodyInstance.InitBody(BodySetup, BodyTransform, nullptr, GetWorld()->GetPhysicsScene(), FInitBodySpawnParams(IsStaticPhysics(), /*bPhysicsTypeDeterminesSimulation*/false), this);
+		}
+	}
+}
+
+void FFastGeoPrimitiveComponent::OnAsyncDestroyPhysicsStateBegin_GameThread()
+{
+	TRACE_CPUPROFILER_EVENT_SCOPE(FFastGeoPrimitiveComponent::OnAsyncDestroyPhysicsStateBegin_GameThread);
+
+	check(!AsyncTermBodyPayload.IsSet());
+	AsyncTermBodyPayload = BodyInstance.StartAsyncTermBody_GameThread();
+	check(!BodyInstance.IsValidBodyInstance());
+
+	Super::OnAsyncDestroyPhysicsStateBegin_GameThread();
+}
+
+void FFastGeoPrimitiveComponent::OnAsyncDestroyPhysicsState()
+{
+	TRACE_CPUPROFILER_EVENT_SCOPE(FFastGeoPrimitiveComponent::OnAsyncDestroyPhysicsState);
+
+	// We tell the BodyInstance to shut down the physics-engine data.
+	if (ensure(AsyncTermBodyPayload.IsSet()))
+	{
+		FBodyInstance::AsyncTermBody(AsyncTermBodyPayload.GetValue());
+		AsyncTermBodyPayload.Reset();
+	}
+
+	Super::OnAsyncDestroyPhysicsState();
+}
+
+bool FFastGeoPrimitiveComponent::IsRenderStateCreated() const
+{
+	return ProxyState == ProxyCreationState::Created;
+}
+
+bool FFastGeoPrimitiveComponent::IsRenderStateDirty() const
+{
+	return bRenderStateDirty;
+}
+
+bool FFastGeoPrimitiveComponent::ShouldCreateRenderState() const
+{
+	return bIsVisible;
+}
+
+bool FFastGeoPrimitiveComponent::IsStaticMeshFirstPersonRelevant() const
+{
+	return GetSceneProxyDesc().IsStaticMeshFirstPersonRelevant();
+}
+
+bool FFastGeoPrimitiveComponent::IsCollisionEnabled() const
+{
+	return GetSceneProxyDesc().bCollisionEnabled;
+}
+
+FSceneInterface* FFastGeoPrimitiveComponent::GetScene() const
+{
+	if (UWorld* World = GetWorld())
+	{
+		return World->Scene;
+	}
+	return nullptr;
+}
+
+FPrimitiveSceneProxy* FFastGeoPrimitiveComponent::GetSceneProxy() const
+{
+	return PrimitiveSceneData.SceneProxy;
+}
+
+void FFastGeoPrimitiveComponent::MarkRenderStateDirty()
+{
+	if ((IsRenderStateCreated() || ProxyState == ProxyCreationState::Delayed) && !IsRenderStateDirty())
+	{
+		bRenderStateDirty = true;
+
+		UFastGeoWorldSubsystem* WorldSubsystem = GetWorld()->GetSubsystem<UFastGeoWorldSubsystem>();
+		if (ensure(WorldSubsystem))
+		{
+			WorldSubsystem->AddToComponentsPendingRecreate(this);
+		}
+	}
+}
+
+FTransform FFastGeoPrimitiveComponent::GetTransform() const
+{
+	return WorldTransform;
+}
+
+FBoxSphereBounds FFastGeoPrimitiveComponent::GetBounds() const
+{
+	return WorldBounds;
+}
+
+FMatrix FFastGeoPrimitiveComponent::GetRenderMatrix() const
+{
+	return GetTransform().ToMatrixWithScale();
+}
+
+float FFastGeoPrimitiveComponent::GetLastRenderTimeOnScreen() const
+{
+	return PrimitiveSceneData.LastRenderTimeOnScreen;
+}
+
+void FFastGeoPrimitiveComponent::SetCollisionEnabled(bool bInCollisionEnabled)
+{
+	FPrimitiveSceneProxyDesc& SceneProxyDesc = GetSceneProxyDesc();
+	SceneProxyDesc.bCollisionEnabled = bInCollisionEnabled;
+}
+
+void FFastGeoPrimitiveComponent::InitializeDynamicProperties()
+{
+	Super::InitializeDynamicProperties();
+
+#if !WITH_EDITOR
+	BodyInstance.FixupData(GetOwnerContainer());
+#endif
+}
+
+#if !WITH_EDITOR && UE_WITH_PSO_PRECACHING
+struct FFastGeoPSOPrecacheFinishedTask
+{
+	explicit FFastGeoPSOPrecacheFinishedTask(FFastGeoPrimitiveComponent* InPrimitiveComponent, int32 InJobSetThatJustCompleted)
+		: WeakPrimitiveComponent(InPrimitiveComponent)
+		, JobSetThatJustCompleted(InJobSetThatJustCompleted)
+	{
+	}
+
+	static TStatId GetStatId() { return TStatId(); }
+	static ENamedThreads::Type GetDesiredThread() { return ENamedThreads::GameThread; }
+	static ESubsequentsMode::Type GetSubsequentsMode() { return ESubsequentsMode::TrackSubsequents; }
+
+	void DoTask(ENamedThreads::Type CurrentThread, const FGraphEventRef& MyCompletionGraphEvent)
+	{
+		if (FFastGeoPrimitiveComponent* PrimitiveComponent = WeakPrimitiveComponent.Get<FFastGeoPrimitiveComponent>())
+		{
+			PrimitiveComponent->OnPrecacheFinished(JobSetThatJustCompleted);
+		}
+	}
+
+	FWeakFastGeoComponent WeakPrimitiveComponent;
+	int32 JobSetThatJustCompleted;
+};
+
+void FFastGeoPrimitiveComponent::OnPrecacheFinished(int32 JobSetThatJustCompleted)
+{
+	QUICK_SCOPE_CYCLE_COUNTER(STAT_PSOPrecacheFinishedTask);
+	int32 CurrJobSetCompleted = LatestPSOPrecacheJobSetCompleted.load();
+	while (CurrJobSetCompleted < JobSetThatJustCompleted && !LatestPSOPrecacheJobSetCompleted.compare_exchange_weak(CurrJobSetCompleted, JobSetThatJustCompleted)) {}
+	MarkRenderStateDirty();
+}
+
+void FFastGeoPrimitiveComponent::RequestRecreateRenderStateWhenPSOPrecacheFinished(const FGraphEventArray& PSOPrecacheCompileEvents)
+{
+	// If the proxy creation strategy relies on knowing when the precached PSO has been compiled,
+	// schedule a task to mark the render state dirty when all PSOs are compiled so the proxy gets recreated.
+	if (GetPSOPrecacheProxyCreationStrategy() != EPSOPrecacheProxyCreationStrategy::AlwaysCreate)
+	{
+		LatestPSOPrecacheJobSet++;
+		if (!PSOPrecacheCompileEvents.IsEmpty())
+		{
+			TGraphTask<FFastGeoPSOPrecacheFinishedTask>::CreateTask(&PSOPrecacheCompileEvents).ConstructAndDispatchWhenReady(this, LatestPSOPrecacheJobSet);
+		}
+		else
+		{
+			// No graph events to wait on, the job set can be considered complete.
+			LatestPSOPrecacheJobSetCompleted = LatestPSOPrecacheJobSet;
+		}
+	}
+	bPSOPrecacheCalled = true;
+}
+
+void FFastGeoPrimitiveComponent::SetupPrecachePSOParams(FPSOPrecacheParams& Params)
+{
+	const FPrimitiveSceneProxyDesc& SceneProxyDesc = GetSceneProxyDesc();
+	auto IsPrecomputedLightingValid = []() { return false; };
+	Params.bRenderInMainPass = SceneProxyDesc.bRenderInMainPass;
+	Params.bRenderInDepthPass = SceneProxyDesc.bRenderInDepthPass;
+	Params.bStaticLighting = SceneProxyDesc.bHasStaticLighting;
+	Params.bUsesIndirectLightingCache = Params.bStaticLighting && SceneProxyDesc.IndirectLightingCacheQuality != ILCQ_Off && (!IsPrecomputedLightingValid() || SceneProxyDesc.LightmapType == ELightmapType::ForceVolumetric);
+	Params.bAffectDynamicIndirectLighting = SceneProxyDesc.bAffectDynamicIndirectLighting;
+	Params.bCastShadow = SceneProxyDesc.CastShadow;
+	// Custom depth can be toggled at runtime with PSO precache call so assume it might be needed when depth pass is needed
+	// Ideally precache those with lower priority and don't wait on these (UE-174426)
+	Params.bRenderCustomDepth = SceneProxyDesc.bRenderInDepthPass;
+	Params.bCastShadowAsTwoSided = SceneProxyDesc.bCastShadowAsTwoSided;
+	Params.SetMobility(SceneProxyDesc.Mobility);
+	Params.SetStencilWriteMask(FRendererStencilMaskEvaluation::ToStencilMask(SceneProxyDesc.CustomDepthStencilWriteMask));
+
+	TArray<UMaterialInterface*> UsedMaterials;
+	GetUsedMaterials(UsedMaterials);
+	for (const UMaterialInterface* MaterialInterface : UsedMaterials)
+	{
+		if (MaterialInterface && MaterialInterface->IsUsingWorldPositionOffset_Concurrent(GMaxRHIFeatureLevel))
+		{
+			Params.bAnyMaterialHasWorldPositionOffset = true;
+			break;
+		}
+	}
+}
+
+#endif
+
+bool FFastGeoPrimitiveComponent::IsPSOPrecaching() const
+{
+#if !WITH_EDITOR && UE_WITH_PSO_PRECACHING
+	return LatestPSOPrecacheJobSetCompleted != LatestPSOPrecacheJobSet;
+#else
+	return false;
+#endif
+}
+
+bool FFastGeoPrimitiveComponent::ShouldRenderProxyFallbackToDefaultMaterial() const
+{
+#if !WITH_EDITOR && UE_WITH_PSO_PRECACHING
+	return IsPSOPrecaching() && GetPSOPrecacheProxyCreationStrategy() == EPSOPrecacheProxyCreationStrategy::UseDefaultMaterialUntilPSOPrecached;
+#else
+	return false;
+#endif
+}
+
+#if !WITH_EDITOR && UE_WITH_PSO_PRECACHING
+bool FFastGeoPrimitiveComponent::CheckPSOPrecachingAndBoostPriority(EPSOPrecachePriority NewPSOPrecachePriority)
+{
+	bool bPrecacheStillRunning = IsPSOPrecaching();
+
+	ensure(!IsComponentPSOPrecachingEnabled() || bPSOPrecacheCalled);
+	check(NewPSOPrecachePriority == EPSOPrecachePriority::High || NewPSOPrecachePriority == EPSOPrecachePriority::Highest);
+
+	if (bPrecacheStillRunning && PSOPrecacheRequestPriority < NewPSOPrecachePriority)
+	{
+		BoostPSOPriority(NewPSOPrecachePriority, MaterialPSOPrecacheRequestIDs);
+		PSOPrecacheRequestPriority = NewPSOPrecachePriority;
+	}
+	return bPrecacheStillRunning;
+}
+#endif
+
+void FFastGeoPrimitiveComponent::PrecachePSOs()
+{
+#if !WITH_EDITOR && UE_WITH_PSO_PRECACHING
+	TRACE_CPUPROFILER_EVENT_SCOPE(FFastGeoPrimitiveComponent::PrecachePSOs);
+
+	if (!FApp::CanEverRender() || !IsComponentPSOPrecachingEnabled())
+	{
+		return;
+	}
+
+	// Clear the current request data
+	MaterialPSOPrecacheRequestIDs.Empty();
+	PSOPrecacheRequestPriority = EPSOPrecachePriority::Medium;
+
+	// Collect the data from the derived classes
+	FPSOPrecacheParams PSOPrecacheParams;
+	SetupPrecachePSOParams(PSOPrecacheParams);
+	FMaterialInterfacePSOPrecacheParamsList PSOPrecacheDataArray;
+	CollectPSOPrecacheData(PSOPrecacheParams, PSOPrecacheDataArray);
+
+	FGraphEventArray GraphEvents;
+	PrecacheMaterialPSOs(PSOPrecacheDataArray, MaterialPSOPrecacheRequestIDs, GraphEvents);
+
+	RequestRecreateRenderStateWhenPSOPrecacheFinished(GraphEvents);
+#endif
+}
+
+//~ Begin IPhysicsBodyInstanceOwner interface
+bool FFastGeoPrimitiveComponent::IsStaticPhysics() const
+{
+	return GetSceneProxyDesc().Mobility != EComponentMobility::Movable && bStaticWhenNotMoveable;
+}
+
+UObject* FFastGeoPrimitiveComponent::GetSourceObject() const
+{
+	return GetOwnerContainer();
+}
+//~ End IPhysicsBodyInstanceOwner interface
+
+
+#if WITH_EDITOR
+void UFastGeoPrimitiveComponentEditorProxy::NotifyRenderStateChanged()
+{
+	FObjectCacheEventSink::NotifyRenderStateChanged_Concurrent(this);
+}
+
+IPrimitiveComponent* UFastGeoPrimitiveComponentEditorProxy::GetPrimitiveComponentInterface()
+{
+	return this;
+}
+
+void UFastGeoPrimitiveComponentEditorProxy::BeginDestroy()
+{
+	Super::BeginDestroy();
+
+	if (!IsTemplate())
+	{
+		NotifyRenderStateChanged();
+	}
+}
+
+//~ Begin IPrimitiveComponent interface
+bool UFastGeoPrimitiveComponentEditorProxy::IsRenderStateCreated() const
+{
+	return GetComponent<ComponentType>().IsRenderStateCreated();
+}
+
+bool UFastGeoPrimitiveComponentEditorProxy::IsRenderStateDirty() const
+{
+	return GetComponent<ComponentType>().IsRenderStateDirty();
+}
+
+bool UFastGeoPrimitiveComponentEditorProxy::ShouldCreateRenderState() const
+{
+	return GetComponent<ComponentType>().ShouldCreateRenderState();
+}
+
+bool UFastGeoPrimitiveComponentEditorProxy::IsRegistered() const
+{
+	return GetComponent<ComponentType>().IsRegistered();
+}
+
+bool UFastGeoPrimitiveComponentEditorProxy::IsUnreachable() const
+{
+	return Super::IsUnreachable();
+}
+
+UWorld* UFastGeoPrimitiveComponentEditorProxy::GetWorld() const
+{
+	return GetComponent<ComponentType>().GetWorld();
+}
+
+FSceneInterface* UFastGeoPrimitiveComponentEditorProxy::GetScene() const
+{
+	return GetComponent<ComponentType>().GetScene();
+}
+
+FPrimitiveSceneProxy* UFastGeoPrimitiveComponentEditorProxy::GetSceneProxy() const
+{
+	return GetComponent<ComponentType>().GetSceneProxy();
+}
+
+void UFastGeoPrimitiveComponentEditorProxy::GetUsedMaterials(TArray<UMaterialInterface*>& OutMaterials, bool bGetDebugMaterials) const
+{
+	GetComponent<ComponentType>().GetUsedMaterials(OutMaterials, bGetDebugMaterials);
+}
+
+void UFastGeoPrimitiveComponentEditorProxy::MarkRenderStateDirty()
+{
+	GetComponent<ComponentType>().MarkRenderStateDirty();
+}
+
+void UFastGeoPrimitiveComponentEditorProxy::DestroyRenderState()
+{
+	GetComponent<ComponentType>().DestroyRenderState(nullptr);
+}
+
+void UFastGeoPrimitiveComponentEditorProxy::CreateRenderState(FRegisterComponentContext* Context)
+{
+	GetComponent<ComponentType>().CreateRenderState(Context);
+}
+
+FString UFastGeoPrimitiveComponentEditorProxy::GetName() const
+{
+	return GetUObject()->GetName();
+}
+
+FString UFastGeoPrimitiveComponentEditorProxy::GetFullName() const
+{
+	return GetUObject()->GetFullName();
+}
+
+FTransform UFastGeoPrimitiveComponentEditorProxy::GetTransform() const
+{
+	return GetComponent<ComponentType>().GetTransform();
+}
+
+FBoxSphereBounds UFastGeoPrimitiveComponentEditorProxy::GetBounds() const
+{
+	return GetComponent<ComponentType>().GetBounds();
+}
+
+float UFastGeoPrimitiveComponentEditorProxy::GetLastRenderTimeOnScreen() const
+{
+	return GetComponent<ComponentType>().GetLastRenderTimeOnScreen();
+}
+
+void UFastGeoPrimitiveComponentEditorProxy::GetPrimitiveStats(FPrimitiveStats& PrimitiveStats) const
+{
+}
+
+UObject* UFastGeoPrimitiveComponentEditorProxy::GetUObject()
+{
+	return this;
+}
+
+const UObject* UFastGeoPrimitiveComponentEditorProxy::GetUObject() const
+{
+	return this;
+}
+
+void UFastGeoPrimitiveComponentEditorProxy::PrecachePSOs()
+{
+	GetComponent<ComponentType>().PrecachePSOs();
+}
+
+UObject* UFastGeoPrimitiveComponentEditorProxy::GetOwner() const
+{
+	return GetOuter();
+}
+
+FString UFastGeoPrimitiveComponentEditorProxy::GetOwnerName() const
+{
+	return GetOwner()->GetName();
+}
+
+FPrimitiveSceneProxy* UFastGeoPrimitiveComponentEditorProxy::CreateSceneProxy()
+{
+	return GetComponent<ComponentType>().CreateSceneProxy();
+}
+
+HHitProxy* UFastGeoPrimitiveComponentEditorProxy::CreateMeshHitProxy(int32 SectionIndex, int32 MaterialIndex)
+{
+	return nullptr;
+}
+
+HHitProxy* UFastGeoPrimitiveComponentEditorProxy::CreatePrimitiveHitProxies(TArray<TRefCountPtr<HHitProxy>>& OutHitProxies)
+{
+	return nullptr;
+}
+//~ End IPrimitiveComponent interface
+#endif
diff --git a/Engine/Plugins/Experimental/FastGeoStreaming/Source/FastGeoStreaming/Private/FastGeoPrimitiveComponent.h b/Engine/Plugins/Experimental/FastGeoStreaming/Source/FastGeoStreaming/Private/FastGeoPrimitiveComponent.h
new file mode 100644
--- /dev/null
+++ b/Engine/Plugins/Experimental/FastGeoStreaming/Source/FastGeoStreaming/Private/FastGeoPrimitiveComponent.h
@@ -0,0 +1,230 @@
+// Copyright Epic Games, Inc. All Rights Reserved.
+
+#pragma once
+
+#include "FastGeoComponent.h"
+#include "Components/PrimitiveComponent.h"
+#include "Components/ComponentInterfaces.h"
+#include "Engine/EngineTypes.h"
+#include "PrimitiveSceneInfoData.h"
+#include "PrimitiveSceneProxy.h"
+#include "RenderCommandFence.h"
+#include "PhysicsEngine/BodyInstance.h"
+#include "PhysicsEngine/PhysicsBodyInstanceOwnerInterface.h"
+#include "PSOPrecacheFwd.h"
+#include "FastGeoPrimitiveComponent.generated.h"
+
+class FPrimitiveSceneProxy;
+class UPrimitiveComponent;
+class URuntimeVirtualTexture;
+class FFastGeoPrimitiveComponent;
+struct FPrimitiveSceneDesc;
+struct FPrimitiveSceneProxyDesc;
+struct FPSOPrecacheParams;
+enum class EPSOPrecachePriority : uint8;
+
+enum class ESceneProxyCreationError
+{
+	None,
+	WaitingPSOs,
+	InvalidMesh
+};
+
+class FASTGEOSTREAMING_API FFastGeoPrimitiveComponent : public FFastGeoComponent, public IPhysicsBodyInstanceOwner
+{
+public:
+	typedef FFastGeoComponent Super;
+	typedef UFastGeoPrimitiveComponentEditorProxy EditorProxyType;
+
+	/** Static type identifier for this element class */
+	static const FFastGeoElementType Type;
+
+	FFastGeoPrimitiveComponent(int32 InComponentIndex = INDEX_NONE, FFastGeoElementType InType = Type);
+	virtual ~FFastGeoPrimitiveComponent() = default;
+
+	TArray<URuntimeVirtualTexture*> const& GetRuntimeVirtualTextures() const { return RuntimeVirtualTextures; }
+	bool IsStaticMeshFirstPersonRelevant() const;
+
+	FPrimitiveComponentId GetPrimitiveSceneId() const { return PrimitiveSceneData.PrimitiveSceneId; }
+	virtual UObject const* AdditionalStatObject() const { return nullptr; }
+		
+	FSceneInterface* GetScene() const;
+	FPrimitiveSceneProxy* GetSceneProxy() const;
+	FTransform GetTransform() const;
+	FBoxSphereBounds GetBounds() const;
+	FMatrix GetRenderMatrix() const;
+	float GetLastRenderTimeOnScreen() const;
+
+	void CreateRenderState(FRegisterComponentContext* Context);
+
+	struct FFastGeoDestroyRenderStateContext
+	{
+		FFastGeoDestroyRenderStateContext(FSceneInterface* InScene);
+		~FFastGeoDestroyRenderStateContext();
+		bool HasPendingWork() const;
+
+		static void DestroyProxy(FFastGeoDestroyRenderStateContext* InContext, FPrimitiveSceneProxy* InPrimitiveSceneProxy);
+
+	private:
+		FSceneInterface* Scene;
+		TArray<FPrimitiveSceneProxy*> PrimitiveSceneProxies;
+	};
+
+	void DestroyRenderState(FFastGeoDestroyRenderStateContext* Context);
+	bool IsRenderStateCreated() const;
+	bool IsRenderStateDirty() const;
+	bool IsVisible() const;
+	bool ShouldCreateRenderState() const;
+	void MarkRenderStateDirty();
+	void PrecachePSOs();
+	bool IsPSOPrecaching() const;
+	void SetCollisionEnabled(bool bInCollisionEnabled);
+	void UpdateVisibility();
+	
+	//~ Begin IPhysicsBodyInstanceOwner interface
+	virtual bool IsStaticPhysics() const override;
+	virtual UObject* GetSourceObject() const override;
+	virtual UPhysicalMaterial* GetPhysicalMaterial() const override { return nullptr; }
+	virtual void GetComplexPhysicalMaterials(TArray<UPhysicalMaterial*>& OutPhysMaterials, TArray<FPhysicalMaterialMaskParams>* OutPhysMaterialMasks) const override {}
+	//~ End IPhysicsBodyInstanceOwner interface
+	
+	//~ Begin FFastGeoComponent interface
+	virtual void Serialize(FArchive& Ar) override;
+	virtual void InitializeDynamicProperties() override;
+	virtual void OnAsyncCreatePhysicsState() override;
+	virtual void OnAsyncDestroyPhysicsStateBegin_GameThread() override;
+	virtual void OnAsyncDestroyPhysicsState() override;
+	virtual bool IsCollisionEnabled() const override;
+#if WITH_EDITOR
+	virtual void InitializeFromComponent(UActorComponent* Component) override;
+	virtual UClass* GetEditorProxyClass() const override;
+#endif
+
+protected:
+	virtual FPrimitiveSceneProxyDesc& GetSceneProxyDesc() = 0;
+	virtual const FPrimitiveSceneProxyDesc& GetSceneProxyDesc() const = 0;
+#if WITH_EDITOR
+	virtual void InitializeSceneProxyDescFromComponent(UActorComponent* Component) = 0;
+	virtual void ResetSceneProxyDescUnsupportedProperties();
+#endif
+	virtual void InitializeSceneProxyDescDynamicProperties();
+	virtual void ApplyWorldTransform(const FTransform& InTransform);
+	virtual FPrimitiveSceneProxy* CreateSceneProxy(ESceneProxyCreationError* OutError = nullptr) = 0;
+	virtual void GetUsedMaterials(TArray<UMaterialInterface*>& OutMaterials, bool bGetDebugMaterials = false) const = 0;
+	FPrimitiveSceneDesc BuildSceneDesc();
+
+protected:
+	// Persistent data
+	FTransform LocalTransform;
+	FTransform WorldTransform;
+	FBoxSphereBounds LocalBounds;
+	FBoxSphereBounds WorldBounds;
+	uint8 bIsVisible : 1 = true;
+	uint8 bStaticWhenNotMoveable : 1 = true;
+	FCustomPrimitiveData CustomPrimitiveData;
+	TEnumAsByte<enum EDetailMode> DetailMode;
+	FBodyInstance BodyInstance;
+	TArray<URuntimeVirtualTexture*> RuntimeVirtualTextures;
+
+	// Runtime Data (transient)
+	FPrimitiveSceneInfoData PrimitiveSceneData{};
+	/** Payload used to release BodyInstance resources in asynchronous mode (see OnAsyncDestroyPhysicsState). */
+	TOptional<FBodyInstance::FAsyncTermBodyPayload> AsyncTermBodyPayload;
+
+	enum class  ProxyCreationState : uint8
+	{
+		None, // Constructed/Initialized
+		Pending, // AddedToWorld & proxy creation is pending
+		Creating, // Actively creating the proxy
+		Created, // Proxy is now created
+		Delayed // Proxy creation delayed (used when PSO precaching is not ready when creating proxy)
+	};
+
+	ProxyCreationState ProxyState = ProxyCreationState::None;
+	bool bRenderStateDirty = false;
+
+private:
+	bool ShouldRenderProxyFallbackToDefaultMaterial() const;
+
+#if !WITH_EDITOR && UE_WITH_PSO_PRECACHING
+public:
+	void OnPrecacheFinished(int32 JobSetThatJustCompleted);
+
+protected:
+	bool CheckPSOPrecachingAndBoostPriority(EPSOPrecachePriority NewPSOPrecachePriority);
+	void RequestRecreateRenderStateWhenPSOPrecacheFinished(const FGraphEventArray& PSOPrecacheCompileEvents);
+	void SetupPrecachePSOParams(FPSOPrecacheParams& Params);
+	virtual void CollectPSOPrecacheData(const FPSOPrecacheParams& BasePrecachePSOParams, FMaterialInterfacePSOPrecacheParamsList& OutParams) {}
+
+	// Cached array of material PSO requests which can be used to boost the priority
+	TArray<FMaterialPSOPrecacheRequestID> MaterialPSOPrecacheRequestIDs;
+	// Atomic int used to track the last PSO precache events
+	std::atomic<int> LatestPSOPrecacheJobSetCompleted = 0;
+	int32 LatestPSOPrecacheJobSet = 0;
+	// Helper flag to check if PSOs have been precached already
+	uint8 bPSOPrecacheCalled : 1 = false;
+	// PSOs requested priority
+	EPSOPrecachePriority PSOPrecacheRequestPriority : 2 = EPSOPrecachePriority::Medium;
+	static_assert((int)EPSOPrecachePriority::Highest < 1 << 2);
+#endif
+
+	friend class UFastGeoContainer;
+	friend class UFastGeoPrimitiveComponentEditorProxy;
+};
+
+// Dummy type to use as base class in non editor builds
+template <typename TInterface>
+class IFastGeoDummyInterface {};
+
+// Implement TInterface, only in editor.
+// Would have been easier with a #ifdef in the class declaration, but UHT doesn't allow it.
+template <typename TInterface>
+class IFastGeoEditorInterface : public std::conditional<WITH_EDITOR, TInterface, IFastGeoDummyInterface<TInterface>>::type {};
+
+UCLASS()
+class UFastGeoPrimitiveComponentEditorProxy : public UFastGeoComponentEditorProxy, public IFastGeoEditorInterface<IPrimitiveComponent>
+{
+	GENERATED_BODY()
+
+#if WITH_EDITOR
+	typedef FFastGeoPrimitiveComponent ComponentType;
+
+public:
+	virtual void NotifyRenderStateChanged();
+	virtual IPrimitiveComponent* GetPrimitiveComponentInterface();
+
+protected:
+	//~ Begin UObject interface
+	virtual void BeginDestroy() override;
+	//~ End UObject interface
+
+	//~ Begin IPrimitiveComponent interface
+	virtual bool IsRenderStateCreated() const override;
+	virtual bool IsRenderStateDirty() const override;
+	virtual bool ShouldCreateRenderState() const override;
+	virtual bool IsRegistered() const override;
+	virtual bool IsUnreachable() const override;
+	virtual UWorld* GetWorld() const override;
+	virtual FSceneInterface* GetScene() const override;
+	virtual FPrimitiveSceneProxy* GetSceneProxy() const override;
+	virtual void GetUsedMaterials(TArray<UMaterialInterface*>& OutMaterials, bool bGetDebugMaterials = false) const override;
+	virtual void MarkRenderStateDirty() override;
+	virtual void DestroyRenderState() override;
+	virtual void CreateRenderState(FRegisterComponentContext* Context) override;
+	virtual FString GetName() const override;
+	virtual FString GetFullName() const override;
+	virtual FTransform GetTransform() const override;
+	virtual FBoxSphereBounds GetBounds() const override;
+	virtual float GetLastRenderTimeOnScreen() const override;
+	virtual void GetPrimitiveStats(FPrimitiveStats& PrimitiveStats) const override;
+	virtual UObject* GetUObject() override;
+	virtual const UObject* GetUObject() const override;
+	virtual void PrecachePSOs() override;
+	virtual UObject* GetOwner() const override;
+	virtual FString GetOwnerName() const override;
+	virtual FPrimitiveSceneProxy* CreateSceneProxy() override;
+	virtual HHitProxy* CreateMeshHitProxy(int32 SectionIndex, int32 MaterialIndex) override;
+	virtual HHitProxy* CreatePrimitiveHitProxies(TArray<TRefCountPtr<HHitProxy> >& OutHitProxies) override;
+	//~ End IPrimitiveComponent interface
+#endif
+};
diff --git a/Engine/Plugins/Experimental/FastGeoStreaming/Source/FastGeoStreaming/Private/FastGeoStaticMeshComponent.cpp b/Engine/Plugins/Experimental/FastGeoStreaming/Source/FastGeoStreaming/Private/FastGeoStaticMeshComponent.cpp
new file mode 100644
--- /dev/null
+++ b/Engine/Plugins/Experimental/FastGeoStreaming/Source/FastGeoStreaming/Private/FastGeoStaticMeshComponent.cpp
@@ -0,0 +1,468 @@
+// Copyright Epic Games, Inc. All Rights Reserved.
+
+#include "FastGeoStaticMeshComponent.h"
+#include "FastGeoContainer.h"
+#include "Materials/MaterialRelevance.h"
+#include "SceneInterface.h"
+#include "MeshComponentHelper.h"
+#include "PrimitiveComponentHelper.h"
+#include "StaticMeshComponentHelper.h"
+#include "Rendering/NaniteResourcesHelper.h"
+#include "PhysicalMaterials/PhysicalMaterial.h"
+#include "NaniteVertexFactory.h"
+#include "StaticMeshSceneProxy.h"
+#include "FastGeoLog.h"
+
+#if WITH_EDITOR
+#include "ObjectCacheEventSink.h"
+#endif
+
+const FFastGeoElementType FFastGeoStaticMeshComponentBase::Type(&FFastGeoPrimitiveComponent::Type);
+
+FFastGeoStaticMeshComponentBase::FFastGeoStaticMeshComponentBase(int32 InComponentIndex, FFastGeoElementType InType)
+	: Super(InComponentIndex, InType)
+{
+}
+
+void FFastGeoStaticMeshComponentBase::Serialize(FArchive& Ar)
+{
+	Super::Serialize(Ar);
+
+	// Serialize persistent data from FFastGeoStaticMeshComponentBase
+	Ar << OverrideMaterials;
+
+	// Serialize persistent data from FStaticMeshSceneProxyDesc
+	FStaticMeshSceneProxyDesc& SceneProxyDesc = GetStaticMeshSceneProxyDesc();
+	Ar << SceneProxyDesc.StaticMesh;
+	Ar << SceneProxyDesc.OverlayMaterial;
+	Ar << SceneProxyDesc.MaterialSlotsOverlayMaterial;
+	Ar << SceneProxyDesc.OverlayMaterialMaxDrawDistance;
+	Ar << SceneProxyDesc.ForcedLodModel;
+	Ar << SceneProxyDesc.MinLOD;
+	Ar << SceneProxyDesc.WorldPositionOffsetDisableDistance;
+	Ar << SceneProxyDesc.NanitePixelProgrammableDistance;
+	Ar << SceneProxyDesc.DistanceFieldSelfShadowBias;
+	Ar << SceneProxyDesc.DistanceFieldIndirectShadowMinVisibility;
+	Ar << SceneProxyDesc.StaticLightMapResolution;
+	FArchive_Serialize_BitfieldBool(Ar, SceneProxyDesc.bReverseCulling);
+	FArchive_Serialize_BitfieldBool(Ar, SceneProxyDesc.bEvaluateWorldPositionOffset);
+	FArchive_Serialize_BitfieldBool(Ar, SceneProxyDesc.bOverrideMinLOD);
+	FArchive_Serialize_BitfieldBool(Ar, SceneProxyDesc.bCastDistanceFieldIndirectShadow);
+	FArchive_Serialize_BitfieldBool(Ar, SceneProxyDesc.bOverrideDistanceFieldSelfShadowBias);
+	FArchive_Serialize_BitfieldBool(Ar, SceneProxyDesc.bEvaluateWorldPositionOffsetInRayTracing);
+	FArchive_Serialize_BitfieldBool(Ar, SceneProxyDesc.bSortTriangles);
+	FArchive_Serialize_BitfieldBool(Ar, SceneProxyDesc.bDisallowNanite);
+	FArchive_Serialize_BitfieldBool(Ar, SceneProxyDesc.bForceDisableNanite);
+	FArchive_Serialize_BitfieldBool(Ar, SceneProxyDesc.bForceNaniteForMasked);
+}
+
+const Nanite::FResources* FFastGeoStaticMeshComponentBase::GetNaniteResources() const
+{
+	if (GetStaticMesh() && GetStaticMesh()->GetRenderData())
+	{
+		return GetStaticMesh()->GetRenderData()->NaniteResourcesPtr.Get();
+	}
+	return nullptr;
+}
+
+FLinearColor FFastGeoStaticMeshComponentBase::GetDebugColorImp() const
+{
+	return GetOwnerComponentCluster()->IsA<FFastGeoHLOD>() ? FColor::Blue : FColor::Cyan;
+}
+
+UBodySetup* FFastGeoStaticMeshComponentBase::GetBodySetup() const
+{
+	return GetStaticMesh() ? GetStaticMesh()->GetBodySetup() : nullptr;
+}
+
+int32 FFastGeoStaticMeshComponentBase::GetNumMaterials() const
+{
+	return GetStaticMesh() ? GetStaticMesh()->GetStaticMaterials().Num() : 0;
+}
+
+FMaterialRelevance FFastGeoStaticMeshComponentBase::GetMaterialRelevance(ERHIFeatureLevel::Type InFeatureLevel) const
+{
+	return FMeshComponentHelper::GetMaterialRelevance(*this, InFeatureLevel);
+}
+
+UMaterialInterface* FFastGeoStaticMeshComponentBase::GetMaterial(int32 MaterialIndex) const
+{
+	return GetMaterial(MaterialIndex, false);
+}
+
+UMaterialInterface* FFastGeoStaticMeshComponentBase::GetNaniteAuditMaterial(int32 MaterialIndex) const
+{
+	return GetMaterial(MaterialIndex, true);
+}
+
+bool FFastGeoStaticMeshComponentBase::UseNaniteOverrideMaterials(bool bDoingMaterialAudit) const
+{
+	return Nanite::FNaniteResourcesHelper::UseNaniteOverrideMaterials(*this, bDoingMaterialAudit);
+}
+
+UMaterialInterface* FFastGeoStaticMeshComponentBase::GetMaterial(int32 MaterialIndex, bool bDoingNaniteMaterialAudit) const
+{
+	TRACE_CPUPROFILER_EVENT_SCOPE(FFastGeoStaticMeshComponentBase::GetMaterial)
+
+	UMaterialInterface* OutMaterial = nullptr;
+
+	// If we have a base materials array, use that
+	if (OverrideMaterials.IsValidIndex(MaterialIndex) && OverrideMaterials[MaterialIndex])
+	{
+		OutMaterial = OverrideMaterials[MaterialIndex];
+	}
+	// Otherwise get from static mesh
+	else if (GetStaticMesh())
+	{
+		OutMaterial = GetStaticMesh()->GetMaterial(MaterialIndex);
+	}
+
+	if (OutMaterial)
+	{
+		//@note FH: temporary preemptive PostLoad until zenloader load ordering improvements
+		// This is still necessary when PrecachePSOs is called in PostLoad (check if we should add a test that we are in the GT).
+		OutMaterial->ConditionalPostLoad();
+
+		// If we have a nanite override, use that
+		if (UseNaniteOverrideMaterials(bDoingNaniteMaterialAudit))
+		{
+			UMaterialInterface* NaniteOverride = OutMaterial->GetNaniteOverride();
+			OutMaterial = NaniteOverride != nullptr ? NaniteOverride : OutMaterial;
+		}
+	}
+
+	return OutMaterial;
+}
+
+const FCollisionResponseContainer& FFastGeoStaticMeshComponentBase::GetCollisionResponseToChannels() const
+{
+	return BodyInstance.GetResponseToChannels();
+}
+
+bool FFastGeoStaticMeshComponentBase::HasValidNaniteData() const
+{
+	return Nanite::FNaniteResourcesHelper::HasValidNaniteData(*this);
+}
+
+#if WITH_EDITOR
+void FFastGeoStaticMeshComponentBase::InitializeSceneProxyDescFromComponent(UActorComponent* Component)
+{
+	UStaticMeshComponent* StaticMeshComponent = CastChecked<UStaticMeshComponent>(Component);
+	FStaticMeshSceneProxyDesc& SceneProxyDesc = GetStaticMeshSceneProxyDesc();
+	SceneProxyDesc.InitializeFromStaticMeshComponent(StaticMeshComponent);
+}
+
+void FFastGeoStaticMeshComponentBase::InitializeFromComponent(UActorComponent* Component)
+{
+	Super::InitializeFromComponent(Component);
+
+	UStaticMeshComponent* StaticMeshComponent = CastChecked<UStaticMeshComponent>(Component);
+	OverrideMaterials = StaticMeshComponent->OverrideMaterials;
+	LocalBounds = GetStaticMesh()->GetBounds();
+	WorldBounds = LocalBounds.TransformBy(WorldTransform);
+}
+
+UClass* FFastGeoStaticMeshComponentBase::GetEditorProxyClass() const
+{
+	return UFastGeoStaticMeshComponentEditorProxy::StaticClass();
+}
+
+void FFastGeoStaticMeshComponentBase::ResetSceneProxyDescUnsupportedProperties()
+{
+	Super::ResetSceneProxyDescUnsupportedProperties();
+
+	// Unsupported properties
+	FStaticMeshSceneProxyDesc& SceneProxyDesc = GetStaticMeshSceneProxyDesc();
+	SceneProxyDesc.LODData = TArrayView<struct FStaticMeshComponentLODInfo>();
+	SceneProxyDesc.LODParentPrimitive = nullptr;
+	SceneProxyDesc.MeshPaintTexture = nullptr;
+	SceneProxyDesc.MeshPaintTextureCoordinateIndex = 0;
+#if STATICMESH_ENABLE_DEBUG_RENDERING
+	SceneProxyDesc.bDrawMeshCollisionIfComplex = false;
+	SceneProxyDesc.bDrawMeshCollisionIfSimple = false;
+#endif
+	SceneProxyDesc.bDisplayNaniteFallbackMesh = false;
+	SceneProxyDesc.SectionIndexPreview = INDEX_NONE;
+	SceneProxyDesc.MaterialIndexPreview = INDEX_NONE;
+	SceneProxyDesc.SelectedEditorMaterial = INDEX_NONE;
+	SceneProxyDesc.SelectedEditorSection = INDEX_NONE;
+
+	// Properties that will be initialized by InitializeSceneProxyDescDynamicProperties
+	SceneProxyDesc.NaniteResources = nullptr;
+	SceneProxyDesc.BodySetup = nullptr;
+	SceneProxyDesc.MaterialRelevance = FMaterialRelevance();
+	SceneProxyDesc.bUseProvidedMaterialRelevance = false;
+}
+#endif
+
+void FFastGeoStaticMeshComponentBase::ApplyWorldTransform(const FTransform& InTransform)
+{
+	Super::ApplyWorldTransform(InTransform);
+
+	WorldBounds = LocalBounds.TransformBy(WorldTransform);
+}
+
+void FFastGeoStaticMeshComponentBase::InitializeSceneProxyDescDynamicProperties()
+{
+	Super::InitializeSceneProxyDescDynamicProperties();
+
+	// Initialize non-serialized properties
+	FStaticMeshSceneProxyDesc& SceneProxyDesc = GetStaticMeshSceneProxyDesc();
+	SceneProxyDesc.OverrideMaterials = OverrideMaterials;
+	SceneProxyDesc.NaniteResources = GetNaniteResources();
+	SceneProxyDesc.BodySetup = GetBodySetup();
+	SceneProxyDesc.SetMaterialRelevance(GetMaterialRelevance(GetScene()->GetFeatureLevel()));
+	SceneProxyDesc.SetCollisionResponseToChannels(GetCollisionResponseToChannels());
+
+	// Add LODData support
+	// Add LODParentPrimitive support
+	// Add MeshPaintTexture/MeshPaintTextureCoordinateIndex support
+}
+
+FPrimitiveSceneProxy* FFastGeoStaticMeshComponentBase::CreateSceneProxy(ESceneProxyCreationError* OutError)
+{
+	check(GetWorld());
+	FSceneInterface* Scene = GetScene();
+	UStaticMesh* StaticMesh = GetStaticMesh();
+	check(Scene);
+	check(StaticMesh);
+	check(StaticMesh->GetRenderData());
+	check(StaticMesh->GetRenderData()->IsInitialized());
+	check(!StaticMesh->IsCompiling()); // Todo: Wait with creation of proxy till stopped compiling.
+
+	if (OutError)
+	{
+		*OutError = ESceneProxyCreationError::None;
+	}
+
+#if !WITH_EDITOR && UE_WITH_PSO_PRECACHING
+	const EPSOPrecachePriority PSOPrecachePriority = GetStaticMeshComponentBoostPSOPrecachePriority();
+	if (CheckPSOPrecachingAndBoostPriority(PSOPrecachePriority) && (GetPSOPrecacheProxyCreationStrategy() == EPSOPrecacheProxyCreationStrategy::DelayUntilPSOPrecached))
+	{
+		UE_LOG(LogFastGeoStreaming, Verbose, TEXT("Skipping CreateSceneProxy for FastGeoStaticMeshComponent of %s (component PSOs are still compiling)"), *GetOwnerComponentCluster()->GetName());
+		if (OutError)
+		{
+			*OutError = ESceneProxyCreationError::WaitingPSOs;
+		}
+		return nullptr;
+	}
+#endif
+
+	InitializeSceneProxyDescDynamicProperties();
+
+	const FStaticMeshSceneProxyDesc& SceneProxyDesc = GetStaticMeshSceneProxyDesc();
+	check(SceneProxyDesc.Scene);
+	check(SceneProxyDesc.Scene == Scene);
+	check(SceneProxyDesc.World == Scene->GetWorld());
+	check(SceneProxyDesc.FeatureLevel == Scene->GetFeatureLevel());
+	check(SceneProxyDesc.ComponentId == GetPrimitiveSceneId());
+
+	Nanite::FMaterialAudit MaterialAudit;
+	if (SceneProxyDesc.ShouldCreateNaniteProxy(&MaterialAudit))
+	{
+		PrimitiveSceneData.SceneProxy = ::new Nanite::FSceneProxy(MaterialAudit, SceneProxyDesc);
+	}
+	else
+	{
+		PrimitiveSceneData.SceneProxy = ::new FStaticMeshSceneProxy(SceneProxyDesc, false);
+	}
+	return PrimitiveSceneData.SceneProxy;
+}
+
+FPrimitiveMaterialPropertyDescriptor FFastGeoStaticMeshComponentBase::GetUsedMaterialPropertyDesc(ERHIFeatureLevel::Type FeatureLevel) const
+{
+	return FPrimitiveComponentHelper::GetUsedMaterialPropertyDesc(*this, FeatureLevel);
+}
+
+void FFastGeoStaticMeshComponentBase::GetUsedMaterials(TArray<UMaterialInterface*>& OutMaterials, bool bGetDebugMaterials) const
+{
+	FStaticMeshComponentHelper::GetUsedMaterials(*this, OutMaterials, bGetDebugMaterials);
+}
+
+void FFastGeoStaticMeshComponentBase::GetDefaultMaterialSlotsOverlayMaterial(TArray<TObjectPtr<UMaterialInterface>>& OutMaterialSlotOverlayMaterials) const
+{
+	//We add null entry for every section of every LOD, this is a requirement for the MeshComponent this class derived from
+	OutMaterialSlotOverlayMaterials.Reset();
+	if (GetStaticMesh())
+	{
+		const TArray<FStaticMaterial>& StaticMaterials = GetStaticMesh()->GetStaticMaterials();
+		for (const FStaticMaterial& StaticMaterial : StaticMaterials)
+		{
+			OutMaterialSlotOverlayMaterials.Add(StaticMaterial.OverlayMaterialInterface);
+		}
+	}
+}
+
+void FFastGeoStaticMeshComponentBase::GetMaterialSlotsOverlayMaterial(TArray<TObjectPtr<UMaterialInterface>>& OutMaterialSlotOverlayMaterials) const
+{
+	//Add all the component override
+	const FStaticMeshSceneProxyDesc& SceneProxyDesc = GetStaticMeshSceneProxyDesc();
+	OutMaterialSlotOverlayMaterials = SceneProxyDesc.MaterialSlotsOverlayMaterial;
+
+	//Get the asset material slot overlay material
+	TArray<TObjectPtr<UMaterialInterface>> AssetMaterialSlotOverlayMaterials;
+	GetDefaultMaterialSlotsOverlayMaterial(AssetMaterialSlotOverlayMaterials);
+
+	//For each slot not override by the component set the asset slot value
+	for (int32 SlotIndex = 0; SlotIndex < AssetMaterialSlotOverlayMaterials.Num(); ++SlotIndex)
+	{
+		if (!OutMaterialSlotOverlayMaterials.IsValidIndex(SlotIndex))
+		{
+			OutMaterialSlotOverlayMaterials.AddZeroed();
+		}
+		//If the component did not already override this slot index put the asset value
+		if (!OutMaterialSlotOverlayMaterials[SlotIndex])
+		{
+			OutMaterialSlotOverlayMaterials[SlotIndex] = AssetMaterialSlotOverlayMaterials[SlotIndex];
+		}
+	}
+}
+
+UStaticMesh* FFastGeoStaticMeshComponentBase::GetStaticMesh() const
+{
+	return GetStaticMeshSceneProxyDesc().StaticMesh;
+}
+
+UMaterialInterface* FFastGeoStaticMeshComponentBase::GetOverlayMaterial() const
+{
+	return GetStaticMeshSceneProxyDesc().OverlayMaterial;
+}
+
+bool FFastGeoStaticMeshComponentBase::IsReverseCulling() const
+{
+	return GetStaticMeshSceneProxyDesc().IsReverseCulling();
+}
+
+bool FFastGeoStaticMeshComponentBase::IsDisallowNanite() const
+{
+	return GetStaticMeshSceneProxyDesc().IsDisallowNanite();
+}
+
+bool FFastGeoStaticMeshComponentBase::IsForceDisableNanite() const
+{
+	return GetStaticMeshSceneProxyDesc().IsForceDisableNanite();
+}
+
+bool FFastGeoStaticMeshComponentBase::IsForceNaniteForMasked() const
+{
+	return GetStaticMeshSceneProxyDesc().IsForceNaniteForMasked();
+}
+
+int32 FFastGeoStaticMeshComponentBase::GetForcedLodModel() const
+{
+	return GetStaticMeshSceneProxyDesc().GetForcedLodModel();
+}
+
+#if WITH_EDITORONLY_DATA
+bool FFastGeoStaticMeshComponentBase::IsDisplayNaniteFallbackMesh() const
+{
+	return GetStaticMeshSceneProxyDesc().IsDisplayNaniteFallbackMesh();
+}
+#endif
+
+//~ Begin IPhysicsBodyInstanceOwner interface
+UPhysicalMaterial* FFastGeoStaticMeshComponentBase::GetPhysicalMaterial() const
+{
+	if (UMaterialInterface* Material = GetMaterial(0))
+	{
+		return Material->GetPhysicalMaterial();
+	}
+	return nullptr;
+}
+
+void FFastGeoStaticMeshComponentBase::GetComplexPhysicalMaterials(TArray<UPhysicalMaterial*>& OutPhysMaterials, TArray<FPhysicalMaterialMaskParams>* OutPhysMaterialMasks) const
+{
+	if (BodyInstance.GetPhysMaterialOverride() != nullptr)
+	{
+		OutPhysMaterials.SetNum(1);
+		OutPhysMaterials[0] = BodyInstance.GetPhysMaterialOverride();
+		check(!OutPhysMaterials[0] || OutPhysMaterials[0]->IsValidLowLevel());
+	}
+	else
+	{
+		// See if the Material has a PhysicalMaterial
+		const int32 NumMaterials = GetNumMaterials();
+		OutPhysMaterials.SetNum(NumMaterials);
+
+		if (OutPhysMaterialMasks)
+		{
+			OutPhysMaterialMasks->SetNum(NumMaterials);
+		}
+
+		for (int32 MatIdx = 0; MatIdx < NumMaterials; MatIdx++)
+		{
+			UPhysicalMaterial* PhysMat = GEngine->DefaultPhysMaterial;
+			UMaterialInterface* Material = GetMaterial(MatIdx);
+			if (Material)
+			{
+				PhysMat = Material->GetPhysicalMaterial();
+			}
+
+			OutPhysMaterials[MatIdx] = PhysMat;
+
+			if (OutPhysMaterialMasks)
+			{
+				UPhysicalMaterialMask* PhysMatMask = nullptr;
+				UMaterialInterface* PhysMatMap = nullptr;
+
+				if (Material)
+				{
+					PhysMatMask = Material->GetPhysicalMaterialMask();
+					if (PhysMatMask)
+					{
+						PhysMatMap = Material;
+					}
+				}
+
+				(*OutPhysMaterialMasks)[MatIdx].PhysicalMaterialMask = PhysMatMask;
+				(*OutPhysMaterialMasks)[MatIdx].PhysicalMaterialMap = PhysMatMap;
+			}
+		}
+	}
+}
+//~ End IPhysicsBodyInstanceOwner interface
+
+
+#if !WITH_EDITOR && UE_WITH_PSO_PRECACHING
+void FFastGeoStaticMeshComponentBase::CollectPSOPrecacheData(const FPSOPrecacheParams& BasePrecachePSOParams, FMaterialInterfacePSOPrecacheParamsList& OutParams)
+{
+	FStaticMeshComponentHelper::CollectPSOPrecacheData(*this, BasePrecachePSOParams, OutParams);
+}
+#endif
+
+#if WITH_EDITOR
+
+void UFastGeoStaticMeshComponentEditorProxy::NotifyRenderStateChanged()
+{
+	Super::NotifyRenderStateChanged();
+
+	FObjectCacheEventSink::NotifyStaticMeshChanged_Concurrent(this);
+}
+
+//~ Begin IStaticMeshComponent interface
+void UFastGeoStaticMeshComponentEditorProxy::OnMeshRebuild(bool bRenderDataChanged)
+{
+}
+
+void UFastGeoStaticMeshComponentEditorProxy::PreStaticMeshCompilation()
+{
+}
+
+void UFastGeoStaticMeshComponentEditorProxy::PostStaticMeshCompilation()
+{
+}
+
+UStaticMesh* UFastGeoStaticMeshComponentEditorProxy::GetStaticMesh() const
+{
+	return GetComponent<ComponentType>().GetStaticMesh();
+}
+//~ End IStaticMeshComponent interface
+
+#endif
+
+const FFastGeoElementType FFastGeoStaticMeshComponent::Type(&FFastGeoStaticMeshComponentBase::Type);
+
+FFastGeoStaticMeshComponent::FFastGeoStaticMeshComponent(int32 InComponentIndex, FFastGeoElementType InType)
+	: Super(InComponentIndex, InType)
+{
+}
\ No newline at end of file
diff --git a/Engine/Plugins/Experimental/FastGeoStreaming/Source/FastGeoStreaming/Private/FastGeoStaticMeshComponent.h b/Engine/Plugins/Experimental/FastGeoStreaming/Source/FastGeoStreaming/Private/FastGeoStaticMeshComponent.h
new file mode 100644
--- /dev/null
+++ b/Engine/Plugins/Experimental/FastGeoStreaming/Source/FastGeoStreaming/Private/FastGeoStaticMeshComponent.h
@@ -0,0 +1,164 @@
+// Copyright Epic Games, Inc. All Rights Reserved.
+
+#pragma once
+
+#include "FastGeoPrimitiveComponent.h"
+#include "StaticMeshSceneProxyDesc.h"
+#include "FastGeoStaticMeshComponent.generated.h"
+
+class UBodySetup;
+class UMaterialInterface;
+class FVertexFactoryType;
+struct FMaterialRelevance;
+struct FStaticMeshLODResources;
+struct FPrimitiveMaterialPropertyDescriptor;
+namespace ERHIFeatureLevel { enum Type : int; }
+namespace Nanite
+{
+	struct FResources;
+	class FNaniteResourcesHelper;
+}
+
+class FASTGEOSTREAMING_API FFastGeoStaticMeshComponentBase : public FFastGeoPrimitiveComponent
+{
+public:
+	typedef FFastGeoPrimitiveComponent Super;
+
+	/** Static type identifier for this element class */
+	static const FFastGeoElementType Type;
+
+	FFastGeoStaticMeshComponentBase(int32 InComponentIndex = INDEX_NONE, FFastGeoElementType InType = Type);
+	virtual ~FFastGeoStaticMeshComponentBase() = default;
+
+	const FCollisionResponseContainer& GetCollisionResponseToChannels() const;
+	UMaterialInterface* GetMaterial(int32 MaterialIndex) const;
+	UMaterialInterface* GetNaniteAuditMaterial(int32 MaterialIndex) const;
+	FMaterialRelevance GetMaterialRelevance(ERHIFeatureLevel::Type InFeatureLevel) const;
+	UMaterialInterface* GetOverlayMaterial() const;
+	void GetMaterialSlotsOverlayMaterial(TArray<TObjectPtr<UMaterialInterface>>& OutMaterialSlotOverlayMaterials) const;
+	const Nanite::FResources* GetNaniteResources() const;
+	int32 GetNumMaterials() const;
+	UStaticMesh* GetStaticMesh() const;
+	bool HasValidNaniteData() const;
+	bool UseNaniteOverrideMaterials(bool bDoingMaterialAudit) const;
+	virtual UObject const* AdditionalStatObject() const override { return GetStaticMesh(); }
+	bool IsReverseCulling() const;
+	bool IsDisallowNanite() const;
+	bool IsForceDisableNanite() const;
+	bool IsForceNaniteForMasked() const;
+	int32 GetForcedLodModel() const;
+#if WITH_EDITORONLY_DATA
+	bool IsDisplayNaniteFallbackMesh() const;
+#endif
+
+	//~ Begin FFastGeoComponent interface
+	virtual void Serialize(FArchive& Ar) override;
+	virtual UBodySetup* GetBodySetup() const override;
+#if WITH_EDITOR
+	virtual void InitializeFromComponent(UActorComponent* Component) override;
+	virtual UClass* GetEditorProxyClass() const override;
+#endif
+	//~ End FFastGeoComponent interface
+
+	//~ Begin IPhysicsBodyInstanceOwner interface
+	virtual UPhysicalMaterial* GetPhysicalMaterial() const override;
+	virtual void GetComplexPhysicalMaterials(TArray<UPhysicalMaterial*>& OutPhysMaterials, TArray<FPhysicalMaterialMaskParams>* OutPhysMaterialMasks) const override;
+	//~ End IPhysicsBodyInstanceOwner interface
+
+protected:
+	//~ Begin FFastGeoComponent interface
+	virtual FLinearColor GetDebugColorImp() const override;
+	//~ End FFastGeoComponent interface
+
+	//~ Begin FFastGeoPrimitiveComponent interface
+#if WITH_EDITOR
+	virtual void InitializeSceneProxyDescFromComponent(UActorComponent* Component) override;
+	virtual void ResetSceneProxyDescUnsupportedProperties() override;
+#endif
+	virtual void InitializeSceneProxyDescDynamicProperties() override;
+	virtual void ApplyWorldTransform(const FTransform& InTransform) override;
+	virtual FPrimitiveSceneProxy* CreateSceneProxy(ESceneProxyCreationError* OutError = nullptr) override;
+	virtual void GetUsedMaterials(TArray<UMaterialInterface*>& OutMaterials, bool bGetDebugMaterials = false) const override;
+	//~ End FFastGeoPrimitiveComponent interface
+
+	virtual FStaticMeshSceneProxyDesc& GetStaticMeshSceneProxyDesc() = 0;
+	virtual const FStaticMeshSceneProxyDesc& GetStaticMeshSceneProxyDesc() const = 0;
+	UMaterialInterface* GetMaterial(int32 MaterialIndex, bool bDoingNaniteMaterialAudit) const;
+	void GetDefaultMaterialSlotsOverlayMaterial(TArray<TObjectPtr<UMaterialInterface>>& OutMaterialSlotOverlayMaterials) const;
+	FPrimitiveMaterialPropertyDescriptor GetUsedMaterialPropertyDesc(ERHIFeatureLevel::Type FeatureLevel) const;
+#if !WITH_EDITOR && UE_WITH_PSO_PRECACHING
+	virtual void CollectPSOPrecacheData(const FPSOPrecacheParams& BasePrecachePSOParams, FMaterialInterfacePSOPrecacheParamsList& OutParams);
+#endif
+
+protected:
+	// Persistent Data
+	TArray<TObjectPtr<UMaterialInterface>> OverrideMaterials;
+
+	// Temporary until LODData support
+	struct FFastGeoStaticMeshComponentLODInfo
+	{
+		FColorVertexBuffer* OverrideVertexColors;
+	};
+	TArray<FFastGeoStaticMeshComponentLODInfo> LODData;
+
+private:
+	friend class FMeshComponentHelper;
+	friend class FPrimitiveComponentHelper;
+	friend class FStaticMeshComponentHelper;
+	friend class FInstancedStaticMeshComponentHelper;
+	friend class Nanite::FNaniteResourcesHelper;
+};
+
+class FASTGEOSTREAMING_API FFastGeoStaticMeshComponent : public FFastGeoStaticMeshComponentBase
+{
+public:
+	typedef FFastGeoStaticMeshComponentBase Super;
+
+	/** Static type identifier for this element class */
+	static const FFastGeoElementType Type;
+
+	FFastGeoStaticMeshComponent(int32 InComponentIndex = INDEX_NONE, FFastGeoElementType InType = Type);
+	virtual ~FFastGeoStaticMeshComponent() = default;
+
+protected:
+	//~ Begin FFastGeoPrimitiveComponent interface
+	virtual FPrimitiveSceneProxyDesc& GetSceneProxyDesc() override { return SceneProxyDesc; }
+	virtual const FPrimitiveSceneProxyDesc& GetSceneProxyDesc() const override { return SceneProxyDesc; }
+	//~ End FFastGeoPrimitiveComponent interface
+
+	//~ Begin FFastGeoStaticMeshComponentBase interface
+	virtual FStaticMeshSceneProxyDesc& GetStaticMeshSceneProxyDesc() override { return SceneProxyDesc; }
+	virtual const FStaticMeshSceneProxyDesc& GetStaticMeshSceneProxyDesc() const override { return SceneProxyDesc; }
+	//~ End FFastGeoStaticMeshComponentBase interface
+
+private:
+	FStaticMeshSceneProxyDesc SceneProxyDesc{};
+};
+
+UCLASS()
+class UFastGeoStaticMeshComponentEditorProxy : public UFastGeoPrimitiveComponentEditorProxy, public IFastGeoEditorInterface<IStaticMeshComponent>
+{
+	GENERATED_BODY()
+
+#if WITH_EDITOR
+	typedef FFastGeoStaticMeshComponentBase ComponentType;
+
+public:
+	//~ Begin UFastGeoPrimitiveComponentEditorProxy interface
+	virtual void NotifyRenderStateChanged() override;
+	//~ End UFastGeoPrimitiveComponentEditorProxy interface
+
+protected:
+	//~ Begin IStaticMeshComponent interface
+	virtual void OnMeshRebuild(bool bRenderDataChanged) override;
+	virtual void PreStaticMeshCompilation() override;
+	virtual void PostStaticMeshCompilation() override;
+	virtual UStaticMesh* GetStaticMesh() const override;
+	
+	virtual IPrimitiveComponent* GetPrimitiveComponentInterface() override
+	{
+		return UFastGeoPrimitiveComponentEditorProxy::GetPrimitiveComponentInterface();
+	}
+	//~ End IStaticMeshComponent interface
+#endif
+};
\ No newline at end of file
diff --git a/Engine/Plugins/Experimental/FastGeoStreaming/Source/FastGeoStreaming/Private/FastGeoStreamingManager.cpp b/Engine/Plugins/Experimental/FastGeoStreaming/Source/FastGeoStreaming/Private/FastGeoStreamingManager.cpp
new file mode 100644
--- /dev/null
+++ b/Engine/Plugins/Experimental/FastGeoStreaming/Source/FastGeoStreaming/Private/FastGeoStreamingManager.cpp
@@ -0,0 +1,55 @@
+// Copyright Epic Games, Inc. All Rights Reserved.
+
+#include "FastGeoStreamingManager.h"
+#include "FastGeoAsyncRenderStateJobQueue.h"
+#include "FastGeoContainer.h"
+#include "FastGeoComponentCluster.h"
+#include "Engine/World.h"
+#include "Physics/Experimental/PhysScene_Chaos.h"
+
+TUniquePtr<FFastGeoAsyncRenderStateJobQueue> FFastGeoStreamingManager::AsyncRenderStateJobQueue;
+
+void FFastGeoStreamingManager::PushAsyncCreateRenderStateJob(UFastGeoContainer* FastGeo)
+{
+	check(IsInGameThread());
+	if (!AsyncRenderStateJobQueue.IsValid())
+	{
+		AsyncRenderStateJobQueue = MakeUnique<FFastGeoAsyncRenderStateJobQueue>(FastGeo->GetWorld()->Scene);
+	}
+	AsyncRenderStateJobQueue->AddJob({ FastGeo, FFastGeoAsyncRenderStateJobQueue::EJobType::CreateRenderState });
+}
+
+void FFastGeoStreamingManager::PushAsyncDestroyRenderStateJob(UFastGeoContainer* FastGeo)
+{
+	check(IsInGameThread());
+	if (!AsyncRenderStateJobQueue.IsValid())
+	{
+		AsyncRenderStateJobQueue = MakeUnique<FFastGeoAsyncRenderStateJobQueue>(FastGeo->GetWorld()->Scene);
+	}
+	AsyncRenderStateJobQueue->AddJob({ FastGeo, FFastGeoAsyncRenderStateJobQueue::EJobType::DestroyRenderState });
+}
+
+void FFastGeoStreamingManager::ProcessAsyncRenderStateJobs(bool bWaitForCompletion)
+{
+	check(IsInGameThread());
+	if (AsyncRenderStateJobQueue.IsValid())
+	{
+		AsyncRenderStateJobQueue->Tick(bWaitForCompletion);
+		if (AsyncRenderStateJobQueue->IsCompleted())
+		{
+			AsyncRenderStateJobQueue.Reset();
+		}
+	}
+}
+
+void FFastGeoStreamingManager::PushAsyncCreatePhysicsStateJobs(UFastGeoContainer* FastGeo)
+{
+	check(FPhysScene::SupportsAsyncPhysicsStateCreation());
+	FastGeo->OnCreatePhysicsStateBegin_GameThread();
+}
+
+void FFastGeoStreamingManager::PushAsyncDestroyPhysicsStateJobs(UFastGeoContainer* FastGeo)
+{
+	check(FPhysScene::SupportsAsyncPhysicsStateDestruction());
+	FastGeo->OnDestroyPhysicsStateBegin_GameThread();
+}
\ No newline at end of file
diff --git a/Engine/Plugins/Experimental/FastGeoStreaming/Source/FastGeoStreaming/Private/FastGeoStreamingManager.h b/Engine/Plugins/Experimental/FastGeoStreaming/Source/FastGeoStreaming/Private/FastGeoStreamingManager.h
new file mode 100644
--- /dev/null
+++ b/Engine/Plugins/Experimental/FastGeoStreaming/Source/FastGeoStreaming/Private/FastGeoStreamingManager.h
@@ -0,0 +1,44 @@
+// Copyright Epic Games, Inc. All Rights Reserved.
+
+#pragma once
+
+#include "FastGeoAsyncRenderStateJobQueue.h"
+
+class UFastGeoContainer;
+
+class FFastGeoStreamingManager
+{
+public:
+	/** 
+	 * Start a job that will create render state asynchronously for a specified FastGeo container. 
+	 * @param FastGeo: FastGeo container.
+	 */
+	FASTGEOSTREAMING_API static void PushAsyncCreateRenderStateJob(UFastGeoContainer* FastGeo);
+	
+	/**
+	 * Start a job that will destroy render state asynchronously for a specified FastGeo container. 
+	 * @param FastGeo: FastGeo container.
+	 */
+	FASTGEOSTREAMING_API static void PushAsyncDestroyRenderStateJob(UFastGeoContainer* FastGeo);
+
+	/** 
+	 * Update progress of asynchronous render state creation and destruction.
+	 * @param bWaitForCompletion : Whether to wait for pending asynchronous tasks to complete.
+	 */
+	FASTGEOSTREAMING_API static void ProcessAsyncRenderStateJobs(bool bWaitForCompletion = false);
+
+	/** 
+	 * Start a job that will create physics state asynchronously for a specified FastGeo container. 
+	 * @param FastGeo: FastGeo container.
+	 */
+	FASTGEOSTREAMING_API static void PushAsyncCreatePhysicsStateJobs(UFastGeoContainer* FastGeo);
+
+	/** 
+	 * Start a job that will destroy physics state asynchronously for a specified FastGeo container. 
+	 * @param FastGeo: FastGeo container.
+	 */
+	FASTGEOSTREAMING_API static void PushAsyncDestroyPhysicsStateJobs(UFastGeoContainer* FastGeo);
+
+private:
+	static TUniquePtr<FFastGeoAsyncRenderStateJobQueue> AsyncRenderStateJobQueue;
+};
\ No newline at end of file
diff --git a/Engine/Plugins/Experimental/FastGeoStreaming/Source/FastGeoStreaming/Private/FastGeoStreamingModule.cpp b/Engine/Plugins/Experimental/FastGeoStreaming/Source/FastGeoStreaming/Private/FastGeoStreamingModule.cpp
new file mode 100644
--- /dev/null
+++ b/Engine/Plugins/Experimental/FastGeoStreaming/Source/FastGeoStreaming/Private/FastGeoStreamingModule.cpp
@@ -0,0 +1,57 @@
+// Copyright Epic Games, Inc. All Rights Reserved.
+
+#include "FastGeoStreamingModule.h"
+#include "FastGeoPrimitiveComponent.h"
+#include "FastGeoStaticMeshComponent.h"
+#include "FastGeoInstancedStaticMeshComponent.h"
+#include "FastGeoWorldSubsystem.h"
+#include "Modules/ModuleManager.h"
+
+IMPLEMENT_MODULE(FFastGeoStreamingModule, FastGeoStreaming);
+
+void FFastGeoStreamingModule::StartupModule()
+{
+#if WITH_EDITOR
+	IPrimitiveComponent::AddImplementer({ UFastGeoStaticMeshComponentEditorProxy::StaticClass(),  [](UObject* Obj)
+	{
+		return static_cast<IPrimitiveComponent*>(static_cast<UFastGeoPrimitiveComponentEditorProxy*>(Obj));
+	}});
+
+	IStaticMeshComponent::AddImplementer({ UFastGeoStaticMeshComponentEditorProxy::StaticClass(), [](UObject* Obj)
+	{
+		return static_cast<IStaticMeshComponent*>(static_cast<UFastGeoStaticMeshComponentEditorProxy*>(Obj));
+	}});
+
+	IPrimitiveComponent::AddImplementer({ UFastGeoInstancedStaticMeshComponentEditorProxy::StaticClass(), [](UObject* Obj)
+	{
+		return static_cast<IPrimitiveComponent*>(static_cast<UFastGeoPrimitiveComponentEditorProxy*>(Obj));
+	}});
+
+	IStaticMeshComponent::AddImplementer({ UFastGeoInstancedStaticMeshComponentEditorProxy::StaticClass(), [](UObject* Obj)
+	{
+		return static_cast<IStaticMeshComponent*>(static_cast<UFastGeoStaticMeshComponentEditorProxy*>(Obj));
+	}});
+#endif
+
+	Handle_OnWorldPreSendAllEndOfFrameUpdates = FWorldDelegates::OnWorldPreSendAllEndOfFrameUpdates.AddRaw(this, &FFastGeoStreamingModule::OnWorldPreSendAllEndOfFrameUpdates);
+}
+
+void FFastGeoStreamingModule::ShutdownModule()
+{
+#if WITH_EDITOR
+	IPrimitiveComponent::RemoveImplementer(UFastGeoStaticMeshComponentEditorProxy::StaticClass());
+	IStaticMeshComponent::RemoveImplementer(UFastGeoStaticMeshComponentEditorProxy::StaticClass());
+	IPrimitiveComponent::RemoveImplementer(UFastGeoInstancedStaticMeshComponentEditorProxy::StaticClass());
+	IStaticMeshComponent::RemoveImplementer(UFastGeoInstancedStaticMeshComponentEditorProxy::StaticClass());
+#endif
+
+	FWorldDelegates::OnWorldPreSendAllEndOfFrameUpdates.Remove(Handle_OnWorldPreSendAllEndOfFrameUpdates);
+}
+
+void FFastGeoStreamingModule::OnWorldPreSendAllEndOfFrameUpdates(UWorld* InWorld)
+{
+	if (UFastGeoWorldSubsystem* WorldSubsystem = InWorld->GetSubsystem<UFastGeoWorldSubsystem>())
+	{
+		WorldSubsystem->ProcessPendingRecreate();
+	}
+}
diff --git a/Engine/Plugins/Experimental/FastGeoStreaming/Source/FastGeoStreaming/Private/FastGeoWeakElement.h b/Engine/Plugins/Experimental/FastGeoStreaming/Source/FastGeoStreaming/Private/FastGeoWeakElement.h
new file mode 100644
--- /dev/null
+++ b/Engine/Plugins/Experimental/FastGeoStreaming/Source/FastGeoStreaming/Private/FastGeoWeakElement.h
@@ -0,0 +1,92 @@
+// Copyright Epic Games, Inc. All Rights Reserved.
+
+#pragma once
+
+#include "FastGeoContainer.h"
+#include "FastGeoComponentCluster.h"
+#include "FastGeoComponent.h"
+#include "UObject/WeakObjectPtrTemplates.h"
+
+class FWeakFastGeoComponentCluster
+{
+public:
+	FWeakFastGeoComponentCluster(FFastGeoComponentCluster* InComponentCluster = nullptr)
+	{
+		if (UFastGeoContainer* Container = InComponentCluster ? InComponentCluster->GetOwnerContainer() : nullptr)
+		{
+			ContainerWeak = Container;
+			ComponentClusterTypeID = InComponentCluster->GetTypeID();
+			ComponentClusterIndex = InComponentCluster->GetComponentClusterIndex();
+		}
+	}
+
+	FORCEINLINE bool operator==(const FWeakFastGeoComponentCluster& Other) const
+	{
+		return ContainerWeak == Other.ContainerWeak &&
+			ComponentClusterTypeID == Other.ComponentClusterTypeID &&
+			ComponentClusterIndex == Other.ComponentClusterIndex;
+	}
+
+	FORCEINLINE FFastGeoComponentCluster* Get() const
+	{
+		if (UFastGeoContainer* Container = ContainerWeak.Get())
+		{
+			return Container->GetComponentCluster(ComponentClusterTypeID, ComponentClusterIndex);
+		}
+		return nullptr;
+	}
+
+	template <class T>
+	FORCEINLINE T* Get() const
+	{
+		FFastGeoComponentCluster* ComponentCluster = Get();
+		return ComponentCluster ? ComponentCluster->CastTo<T>() : nullptr;
+	}
+
+private:
+	TWeakObjectPtr<UFastGeoContainer> ContainerWeak;
+	uint32 ComponentClusterTypeID = INDEX_NONE;
+	int32 ComponentClusterIndex = INDEX_NONE;
+};
+
+class FWeakFastGeoComponent
+{
+public:
+	FWeakFastGeoComponent(FFastGeoComponent* InComponent)
+	{
+		if (FFastGeoComponentCluster* ComponentCluster = InComponent ? InComponent->GetOwnerComponentCluster() : nullptr)
+		{
+			ComponentClusterWeak = ComponentCluster;
+			ComponentTypeID = InComponent->GetTypeID();
+			ComponentIndex = InComponent->GetComponentIndex();
+		}
+	}
+
+	FORCEINLINE bool operator==(const FWeakFastGeoComponent& Other) const
+	{
+		return ComponentClusterWeak == Other.ComponentClusterWeak &&
+			ComponentTypeID == Other.ComponentTypeID &&
+			ComponentIndex == Other.ComponentIndex;
+	}
+
+	FORCEINLINE FFastGeoComponent* Get() const
+	{
+		if (FFastGeoComponentCluster* ComponentCluster = ComponentClusterWeak.Get())
+		{
+			return ComponentCluster->GetComponent(ComponentTypeID, ComponentIndex);
+		}
+		return nullptr;
+	}
+
+	template <class T>
+	FORCEINLINE T* Get() const
+	{
+		FFastGeoComponent* Component = Get();
+		return Component ? Component->CastTo<T>() : nullptr;
+	}
+
+private:
+	FWeakFastGeoComponentCluster ComponentClusterWeak;
+	uint32 ComponentTypeID = INDEX_NONE;
+	int32 ComponentIndex = INDEX_NONE;
+};
diff --git a/Engine/Plugins/Experimental/FastGeoStreaming/Source/FastGeoStreaming/Private/FastGeoWorldPartitionRuntimeCellTransformer.cpp b/Engine/Plugins/Experimental/FastGeoStreaming/Source/FastGeoStreaming/Private/FastGeoWorldPartitionRuntimeCellTransformer.cpp
new file mode 100644
--- /dev/null
+++ b/Engine/Plugins/Experimental/FastGeoStreaming/Source/FastGeoStreaming/Private/FastGeoWorldPartitionRuntimeCellTransformer.cpp
@@ -0,0 +1,319 @@
+// Copyright Epic Games, Inc. All Rights Reserved.
+
+#include "FastGeoWorldPartitionRuntimeCellTransformer.h"
+
+#if WITH_EDITOR
+#include "CoreMinimal.h"
+#include "GameFramework/Actor.h"
+#include "Engine/Level.h"
+#include "UObject/Package.h"
+#include "Components/StaticMeshComponent.h"
+#include "Components/InstancedStaticMeshComponent.h"
+#include "FastGeoContainer.h"
+#include "FastGeoComponentCluster.h"
+#include "FastGeoComponent.h"
+#include "FastGeoHLOD.h"
+#include "FastGeoStaticMeshComponent.h"
+#include "FastGeoInstancedStaticMeshComponent.h"
+#include "Physics/Experimental/PhysScene_Chaos.h"
+#include "Streaming/ActorTextureStreamingBuildDataComponent.h"
+#include "FastGeoLog.h"
+#endif
+
+#if WITH_EDITORONLY_DATA
+#include "Engine/StaticMeshActor.h"
+#include "LevelInstance/LevelInstanceComponent.h"
+#include "PackedLevelActor/PackedLevelActor.h"
+#include "WorldPartition/HLOD/HLODActor.h"
+#endif
+
+#define LOCTEXT_NAMESPACE "WorldPartition"
+
+UFastGeoWorldPartitionRuntimeCellTransformer::UFastGeoWorldPartitionRuntimeCellTransformer(const FObjectInitializer& ObjectInitializer)
+	: Super(ObjectInitializer)
+{
+#if WITH_EDITORONLY_DATA
+	AllowedClasses.Add(APackedLevelActor::StaticClass());
+	AllowedClasses.Add(AStaticMeshActor::StaticClass());
+	AllowedClasses.Add(AWorldPartitionHLOD::StaticClass());
+	IgnoredComponentClasses.Add(ULevelInstanceComponent::StaticClass());
+	IgnoredComponentClasses.Add(UActorTextureStreamingBuildDataComponent::StaticClass());
+#endif
+}
+
+static const FName NAME_FastGeo(TEXT("FastGeo"));
+static const FName NAME_NoFastGeo(TEXT("NoFastGeo"));
+
+#if WITH_EDITOR
+
+namespace FastGeo
+{
+	auto IsCollisionEnabled = [](UStaticMeshComponent* Component)
+	{
+		return Component->IsCollisionEnabled() && !Component->GetOwner()->IsA<AWorldPartitionHLOD>();
+	};
+
+	class FNullArchive : public FArchive
+	{
+		//override FObjectPtr to do nothing
+		virtual FArchive& operator<<(struct FObjectPtr& Value) override
+		{
+			return *this;
+		}
+	};
+
+	static bool bFastGeoEnabled = true;
+	static FAutoConsoleVariableRef CVarFastGeoEnabled(
+		TEXT("FastGeo.Enable"),
+		bFastGeoEnabled,
+		TEXT("Set to false to disable FastGeoStreaming (used when starting PIE and at cook time)."),
+		ECVF_Default);
+}
+
+void UFastGeoWorldPartitionRuntimeCellTransformer::Transform(ULevel* InLevel)
+{
+	if (!FastGeo::bFastGeoEnabled)
+	{
+		return;
+	}
+
+	if (!FPhysScene::SupportsAsyncPhysicsStateCreation() || 
+		!FPhysScene::SupportsAsyncPhysicsStateDestruction())
+	{
+		UE_LOG(LogFastGeoStreaming, Warning, TEXT("FastGeoStreaming Cell Transformer requires 'p.Chaos.EnableAsyncInitBody' to be enabled."));
+		return;
+	}
+
+	check(InLevel);
+
+	TMap<TObjectPtr<AActor>*, TArray<UStaticMeshComponent*>> ActorComponents;
+	for (TObjectPtr<AActor>& Actor : InLevel->Actors)
+	{
+		if (IsValid(Actor) && CanTransformActor(Actor))
+		{
+			// Gather components
+			Actor->ForEachComponent<UStaticMeshComponent>(false, [&Actor, &ActorComponents](UStaticMeshComponent* StaticMeshComponent)
+			{
+				if (!StaticMeshComponent->IsEditorOnly() &&
+					StaticMeshComponent->GetStaticMesh() &&
+					!StaticMeshComponent->GetLODParentPrimitive() &&
+					(StaticMeshComponent->Mobility == EComponentMobility::Static))
+				{
+					// Make sure BodyInstance CollisionEnabled is updated first
+					StaticMeshComponent->UpdateCollisionFromStaticMesh();
+					const bool bIsCollisionEnabled = FastGeo::IsCollisionEnabled(StaticMeshComponent);
+
+					// If collision is enabled, only allow if async physics state creation and destruction are supported 
+					// If colliion is disabled, allow if component is visible
+					if (!bIsCollisionEnabled ||
+						(StaticMeshComponent->AllowsAsyncPhysicsStateCreation() && StaticMeshComponent->AllowsAsyncPhysicsStateDestruction()))
+					{
+						if (bIsCollisionEnabled || StaticMeshComponent->IsVisible())
+						{
+							ActorComponents.FindOrAdd(&Actor).Add(StaticMeshComponent);
+						}
+					}
+				}
+			});
+		}
+	}
+
+	const TMap<TSubclassOf<UPrimitiveComponent>, FFastGeoElementType> FastGeoComponentTypeMapping =
+	{
+		{ UStaticMeshComponent::StaticClass(), FFastGeoStaticMeshComponent::Type },
+		{ UInstancedStaticMeshComponent::StaticClass(), FFastGeoInstancedStaticMeshComponent::Type }
+	};
+
+	// Walk the component class hierarchy and look for a fast geo mapping
+	auto GetFastGeoComponentType = [&FastGeoComponentTypeMapping] (TSubclassOf<UPrimitiveComponent> InClass) -> FFastGeoElementType
+	{		
+		for (UClass* Class = InClass; Class; Class = Class->GetSuperClass())
+		{
+			if (const FFastGeoElementType* Found = FastGeoComponentTypeMapping.Find(Class))
+			{
+				return *Found;
+			}
+		}
+
+		unimplemented(); // Should never be reached, missing type handling
+		return FFastGeoComponent::Type;
+	};
+
+	if (!ActorComponents.IsEmpty())
+	{
+		const FString CellName = InLevel->GetWorldPartitionRuntimeCell() ? Cast<UObject>(InLevel->GetWorldPartitionRuntimeCell())->GetName() : TEXT("Cell");
+		UFastGeoContainer* FastGeo = NewObject<UFastGeoContainer>(InLevel, *FString::Printf(TEXT("FastGeoContainer_%s"), *CellName));
+		InLevel->AddAssetUserData(FastGeo);
+
+		TUniquePtr<FFastGeoComponentCluster> LevelComponentCluster(new FFastGeoComponentCluster(FastGeo, *FString::Printf(TEXT("FastGeoComponentCluster_%s"), *CellName)));
+
+		for (auto& [Actor, Components] : ActorComponents)
+		{
+			check(!Components.IsEmpty());
+
+			FFastGeoComponentCluster* CurrentComponentCluster = LevelComponentCluster.Get();
+			TUniquePtr<FFastGeoHLOD> FastGeoHLOD;
+			if (AWorldPartitionHLOD* HLODActor = Cast<AWorldPartitionHLOD>(*Actor))
+			{
+				FastGeoHLOD.Reset(new FFastGeoHLOD(FastGeo, *FString::Printf(TEXT("FastGeoHLOD_%s"), *(*Actor)->GetName())));
+				FastGeoHLOD->SetSourceCellGuid(HLODActor->GetSourceCellGuid());
+				FastGeoHLOD->SetRequireWarmup(HLODActor->DoesRequireWarmup());
+				CurrentComponentCluster = FastGeoHLOD.Get();
+			}
+
+			for (UStaticMeshComponent* StaticMeshComponent : Components)
+			{
+				FFastGeoElementType FastGeoComponentType = GetFastGeoComponentType(StaticMeshComponent->GetClass());
+				FFastGeoComponent& FastGeoComponent = CurrentComponentCluster->AddComponent(FastGeoComponentType);
+				FastGeoComponent.InitializeFromComponent(StaticMeshComponent);
+
+				// Remove the component from the actor
+				(*Actor)->RemoveOwnedComponent(StaticMeshComponent);
+				StaticMeshComponent->MarkAsGarbage();
+			}
+
+			if (CanRemoveActor(*Actor))
+			{
+				*Actor = nullptr;
+			}
+
+			if (FastGeoHLOD.IsValid())
+			{
+				check(FastGeoHLOD->HasComponents());
+				FastGeo->AddComponentCluster(FastGeoHLOD.Get());
+			}
+		}
+
+		// Add level component cluster (if not empty)
+		if (LevelComponentCluster->HasComponents())
+		{
+			FastGeo->AddComponentCluster(LevelComponentCluster.Get());
+		}
+
+		// Finalize post-creation intialization
+		FastGeo->OnCreated();
+	}
+
+	InLevel->Actors.Remove(nullptr);
+}
+
+bool UFastGeoWorldPartitionRuntimeCellTransformer::CanTransformActor(AActor* InActor) const
+{
+	if (InActor->GetIsReplicated())
+	{
+		return false;
+	}
+
+	if (!InActor->IsRootComponentStatic())
+	{
+		return false;
+	}
+
+	if (InActor->IsHidden())
+	{
+		return false;
+	}
+
+	if (InActor->IsEditorOnly())
+	{
+		return false;
+	}
+
+	if (InActor->Children.Num() || InActor->IsChildActor())
+	{
+		return false;
+	}
+
+	UClass* ActorClass = InActor->GetClass();
+	for (TSubclassOf<AActor> DisallowedClass : DisallowedClasses)
+	{
+		if (ActorClass == DisallowedClass)
+		{
+			return false;
+		}
+	}
+
+	if (InActor->Tags.Contains(NAME_NoFastGeo))
+	{
+		return false;
+	}
+
+	for (TSubclassOf<AActor> AllowedClass : AllowedClasses)
+	{
+		if (ActorClass->IsChildOf(AllowedClass))
+		{
+			return true;
+		}
+	}
+
+	if (InActor->Tags.Contains(NAME_FastGeo))
+	{
+		return true;
+	}
+
+	return false;
+}
+
+bool UFastGeoWorldPartitionRuntimeCellTransformer::CanRemoveActor(AActor* InActor) const
+{
+	auto IsComponentIgnored = [this](UActorComponent* Component)
+	{
+		if (Component->IsEditorOnly())
+		{
+			return true;
+		}
+
+		if (UStaticMeshComponent* StaticMeshComponent = Cast<UStaticMeshComponent>(Component))
+		{
+			if (!StaticMeshComponent->GetStaticMesh())
+			{
+				return true;
+			}
+
+			if (!StaticMeshComponent->IsVisible() && !FastGeo::IsCollisionEnabled(StaticMeshComponent))
+			{
+				return true;
+			}
+		}
+
+		UClass* ComponentClass = Component->GetClass();
+		for (const TSubclassOf<UActorComponent>& IgnoredComponentClass : IgnoredComponentClasses)
+		{
+			if (ComponentClass->IsChildOf(IgnoredComponentClass))
+			{
+				return true;
+			}
+		}
+		return false;
+	};
+
+	uint32 NumRuntimeComponents = 0;
+	for (UActorComponent* Component : InActor->GetComponents())
+	{
+		if (!IsComponentIgnored(Component))
+		{
+			NumRuntimeComponents++;
+		}
+	}
+		
+	if (!NumRuntimeComponents)
+	{
+		return true;
+	}
+
+	if (NumRuntimeComponents == 1)
+	{
+		if (USceneComponent* RootComponent = InActor->GetRootComponent())
+		{
+			if (ExactCast<USceneComponent>(RootComponent))
+			{
+				return true;
+			}
+		}
+	}
+
+	return false;
+}
+#endif
+
+#undef LOCTEXT_NAMESPACE
diff --git a/Engine/Plugins/Experimental/FastGeoStreaming/Source/FastGeoStreaming/Private/FastGeoWorldPartitionRuntimeCellTransformer.h b/Engine/Plugins/Experimental/FastGeoStreaming/Source/FastGeoStreaming/Private/FastGeoWorldPartitionRuntimeCellTransformer.h
new file mode 100644
--- /dev/null
+++ b/Engine/Plugins/Experimental/FastGeoStreaming/Source/FastGeoStreaming/Private/FastGeoWorldPartitionRuntimeCellTransformer.h
@@ -0,0 +1,37 @@
+// Copyright Epic Games, Inc. All Rights Reserved.
+
+#pragma once
+
+#include "WorldPartition/WorldPartitionRuntimeCellTransformer.h"
+#include "Templates/SubclassOf.h"
+#include "FastGeoWorldPartitionRuntimeCellTransformer.generated.h"
+
+class UActorComponent;
+
+UCLASS()
+class UFastGeoWorldPartitionRuntimeCellTransformer : public UWorldPartitionRuntimeCellTransformer
+{
+	GENERATED_UCLASS_BODY()
+
+#if WITH_EDITOR
+public:
+	virtual void Transform(ULevel* InLevel) override;
+
+private:
+	bool CanTransformActor(AActor* InActor) const;
+	bool CanRemoveActor(AActor* InActor) const;
+#endif
+
+#if WITH_EDITORONLY_DATA
+	/** Allowed classes (recursive) to convert to instances */
+	UPROPERTY(EditAnywhere, Category = FastGeo)
+	TArray<TSubclassOf<AActor>> AllowedClasses;
+
+	/** Disallowed classes (non-recursive) to convert to instances */
+	UPROPERTY(EditAnywhere, Category = FastGeo)
+	TArray<TSubclassOf<AActor>> DisallowedClasses;
+
+	UPROPERTY(EditAnywhere, Category = FastGeo)
+	TArray<TSubclassOf<UActorComponent>> IgnoredComponentClasses;
+#endif
+};
\ No newline at end of file
diff --git a/Engine/Plugins/Experimental/FastGeoStreaming/Source/FastGeoStreaming/Private/FastGeoWorldSubsystem.cpp b/Engine/Plugins/Experimental/FastGeoStreaming/Source/FastGeoStreaming/Private/FastGeoWorldSubsystem.cpp
new file mode 100644
--- /dev/null
+++ b/Engine/Plugins/Experimental/FastGeoStreaming/Source/FastGeoStreaming/Private/FastGeoWorldSubsystem.cpp
@@ -0,0 +1,354 @@
+// Copyright Epic Games, Inc. All Rights Reserved.
+
+#include "FastGeoWorldSubsystem.h"
+#include "FastGeoStreamingManager.h"
+#include "FastGeoContainer.h"
+#include "FastGeoPrimitiveComponent.h"
+#include "Physics/Experimental/PhysScene_Chaos.h"
+#include "Streaming/LevelStreamingDelegates.h"
+#include "Engine/World.h"
+#include "Engine/Level.h"
+#include "WorldPartition/HLOD/HLODRuntimeSubsystem.h"
+#include "WorldPartition/WorldPartitionRuntimeCell.h"
+#include "GameFramework/ActorPrimitiveColorHandler.h"
+#include "Engine/GameViewportClient.h"
+
+#define LOCTEXT_NAMESPACE "FastGeoWorldSubsystem"
+
+static FName NAME_FastGeoColorHandler(TEXT("FastGeo"));
+bool UFastGeoWorldSubsystem::bEnableDebugView = false;
+
+bool UFastGeoWorldSubsystem::IsEnableDebugView()
+{
+	return bEnableDebugView;
+}
+
+UFastGeoWorldSubsystem::UFastGeoWorldSubsystem()
+{
+#if ENABLE_ACTOR_PRIMITIVE_COLOR_HANDLER
+
+	auto UpdatePrimitivesColor = []()
+	{
+		if (GEngine)
+		{
+			for (const FWorldContext& Context : GEngine->GetWorldContexts())
+			{
+				UWorld* World = Context.World();
+				if (World && World->IsGameWorld())
+				{
+					for (ULevel* Level : World->GetLevels())
+					{
+						if (UFastGeoContainer* FastGeo = Level->GetAssetUserData<UFastGeoContainer>())
+						{
+							FastGeo->ForEachComponentCluster([](const FFastGeoComponentCluster& ComponentCluster)
+							{
+								ComponentCluster.ForEachComponent<FFastGeoPrimitiveComponent>([](const FFastGeoPrimitiveComponent& Component)
+								{
+									if (FPrimitiveSceneProxy* SceneProxy = Component.GetSceneProxy())
+									{
+										SceneProxy->SetPrimitiveColor_GameThread(Component.GetDebugColor());
+									}
+								});
+							});
+						}
+					}
+				}
+			}
+		}
+	};
+
+	if (HasAnyFlags(RF_ClassDefaultObject) && ExactCast<UFastGeoWorldSubsystem>(this))
+	{
+		FActorPrimitiveColorHandler::FPrimitiveColorHandler FastGeoColorHandler;
+		FastGeoColorHandler.HandlerName = NAME_FastGeoColorHandler;
+		FastGeoColorHandler.HandlerText = LOCTEXT("FastGeo", "FastGeo");
+		FastGeoColorHandler.HandlerToolTipText = LOCTEXT("FastGeoColor_ToopTip", "Colorize FastGeo primitives. ISM [Orange], HLOD ISM [Red], SM [Cyan], HLOD SM [Blue], else White.");
+		FastGeoColorHandler.bAvailalbleInEditor = false;
+		FastGeoColorHandler.GetColorFunc = [](const UPrimitiveComponent* InPrimitiveComponent)
+		{
+			return FLinearColor::White;
+		};
+		FastGeoColorHandler.ActivateFunc = [&UpdatePrimitivesColor]()
+		{
+			bEnableDebugView = true;
+			UpdatePrimitivesColor();
+		};
+		FastGeoColorHandler.DeactivateFunc = [&UpdatePrimitivesColor]()
+		{
+			bEnableDebugView = false;
+			UpdatePrimitivesColor();
+		};
+		FActorPrimitiveColorHandler::Get().RegisterPrimitiveColorHandler(FastGeoColorHandler);
+	}
+#endif
+}
+
+UFastGeoWorldSubsystem::~UFastGeoWorldSubsystem()
+{
+#if ENABLE_ACTOR_PRIMITIVE_COLOR_HANDLER
+	if (HasAnyFlags(RF_ClassDefaultObject) && ExactCast<UFastGeoWorldSubsystem>(this))
+	{
+		FActorPrimitiveColorHandler::Get().UnregisterPrimitiveColorHandler(NAME_FastGeoColorHandler);
+	}
+#endif
+}
+
+bool UFastGeoWorldSubsystem::DoesSupportWorldType(const EWorldType::Type WorldType) const
+{
+	return WorldType == EWorldType::Game || WorldType == EWorldType::PIE;
+}
+
+void UFastGeoWorldSubsystem::PostInitialize()
+{
+	Super::PostInitialize();
+
+	UWorld* World = GetWorld();
+	if (World->IsPartitionedWorld())
+	{
+		World->OnAllLevelsChanged().AddUObject(this, &UFastGeoWorldSubsystem::OnUpdateLevelStreaming);
+		World->OnAddLevelToWorldExtension().AddUObject(this, &UFastGeoWorldSubsystem::OnAddLevelToWorldExtension);
+		World->OnRemoveLevelFromWorldExtension().AddUObject(this, &UFastGeoWorldSubsystem::OnRemoveLevelFromWorldExtension);
+		FWorldDelegates::LevelComponentsUpdated.AddUObject(this, &UFastGeoWorldSubsystem::OnLevelComponentsUpdated);
+		FWorldDelegates::LevelComponentsCleared.AddUObject(this, &UFastGeoWorldSubsystem::OnLevelComponentsCleared);
+#if DO_CHECK
+		FWorldDelegates::LevelAddedToWorld.AddUObject(this, &UFastGeoWorldSubsystem::OnLevelAddedToWorld);
+		FWorldDelegates::LevelRemovedFromWorld.AddUObject(this, &UFastGeoWorldSubsystem::OnLevelRemovedFromWorld);
+#endif
+		Handle_OnLevelBeginAddToWorld = FLevelStreamingDelegates::OnLevelBeginMakingVisible.AddUObject(this, &UFastGeoWorldSubsystem::OnLevelStartedAddToWorld);
+		Handle_OnLevelBeginRemoveFromWorld = FLevelStreamingDelegates::OnLevelBeginMakingInvisible.AddUObject(this, &UFastGeoWorldSubsystem::OnLevelStartedRemoveFromWorld);
+
+		Handle_OnForEachHLODObjectInCell = World->GetSubsystem<UWorldPartitionHLODRuntimeSubsystem>()->GetForEachHLODObjectInCellEvent().AddUObject(this, &UFastGeoWorldSubsystem::ForEachHLODObjectInCell);
+	}
+}
+
+void UFastGeoWorldSubsystem::Deinitialize()
+{
+	UWorld* World = GetWorld();
+	World->OnAllLevelsChanged().RemoveAll(this);
+	World->OnAddLevelToWorldExtension().RemoveAll(this);
+	World->OnRemoveLevelFromWorldExtension().RemoveAll(this);
+	FWorldDelegates::LevelComponentsUpdated.RemoveAll(this);
+	FWorldDelegates::LevelComponentsCleared.RemoveAll(this);
+#if DO_CHECK
+	FWorldDelegates::LevelAddedToWorld.RemoveAll(this);
+	FWorldDelegates::LevelRemovedFromWorld.RemoveAll(this);
+#endif
+	FLevelStreamingDelegates::OnLevelBeginMakingVisible.Remove(Handle_OnLevelBeginAddToWorld);
+	FLevelStreamingDelegates::OnLevelBeginMakingInvisible.Remove(Handle_OnLevelBeginRemoveFromWorld);
+	Handle_OnLevelBeginAddToWorld.Reset();
+	Handle_OnLevelBeginRemoveFromWorld.Reset();
+
+	World->GetSubsystem<UWorldPartitionHLODRuntimeSubsystem>()->GetForEachHLODObjectInCellEvent().Remove(Handle_OnForEachHLODObjectInCell);
+	Handle_OnForEachHLODObjectInCell.Reset();
+
+	Super::Deinitialize();
+}
+
+void UFastGeoWorldSubsystem::OnLevelStartedAddToWorld(UWorld* World, const ULevelStreaming* StreamingLevel, ULevel* Level)
+{
+	if (World != GetWorld())
+	{
+		return;
+	}
+
+	if (UFastGeoContainer* FastGeo = Level->GetAssetUserData<UFastGeoContainer>())
+	{
+		TRACE_CPUPROFILER_EVENT_SCOPE(UFastGeoWorldSubsystem::OnLevelStartedAddToWorld);
+
+		FastGeo->Register();
+	}
+}
+
+void UFastGeoWorldSubsystem::OnLevelStartedRemoveFromWorld(UWorld* World, const ULevelStreaming* StreamingLevel, ULevel* Level)
+{
+	if (World != GetWorld())
+	{
+		return;
+	}
+
+	if (UFastGeoContainer* FastGeo = Level->GetAssetUserData<UFastGeoContainer>())
+	{
+		TRACE_CPUPROFILER_EVENT_SCOPE(UFastGeoWorldSubsystem::OnLevelStartedRemoveFromWorld);
+
+		FastGeo->Unregister();
+	}
+}
+
+void UFastGeoWorldSubsystem::OnLevelComponentsUpdated(UWorld* World, ULevel* Level)
+{
+	if (World != GetWorld())
+	{
+		return;
+	}
+
+	if (UFastGeoContainer* FastGeo = Level->GetAssetUserData<UFastGeoContainer>())
+	{
+		TRACE_CPUPROFILER_EVENT_SCOPE(UFastGeoWorldSubsystem::OnLevelComponentsUpdated);
+
+		FastGeo->Register();
+
+		const bool bWaitForCompletion = true;
+		FastGeo->Tick(bWaitForCompletion);
+	}
+}
+
+void UFastGeoWorldSubsystem::OnLevelComponentsCleared(UWorld* World, ULevel* Level)
+{
+	if (World != GetWorld())
+	{
+		return;
+	}
+
+	if (UFastGeoContainer* FastGeo = Level->GetAssetUserData<UFastGeoContainer>())
+	{
+		TRACE_CPUPROFILER_EVENT_SCOPE(UFastGeoWorldSubsystem::OnLevelComponentsCleared);
+
+		if (World->IsBeingCleanedUp())
+		{
+			FastGeo->Unregister();
+
+			const bool bWaitForCompletion = true;
+			FastGeo->Tick(bWaitForCompletion);
+		}
+		else
+		{
+			check(!FastGeo->IsRegistered());
+			check(!FastGeo->HasAnyPendingTasks());
+		}
+	}
+}
+
+void UFastGeoWorldSubsystem::OnAddLevelToWorldExtension(ULevel* Level, const bool bWaitForCompletion, bool& bOutHasCompleted)
+{
+	if (UFastGeoContainer* FastGeo = Level->GetAssetUserData<UFastGeoContainer>())
+	{
+		TRACE_CPUPROFILER_EVENT_SCOPE(UFastGeoWorldSubsystem::OnAddLevelToWorldExtension);
+
+		if (FastGeo->HasAnyPendingCreateTasks())
+		{
+			FastGeo->Tick(bWaitForCompletion);
+		}
+
+		if (FastGeo->HasAnyPendingCreateTasks())
+		{
+			bOutHasCompleted = false;
+		}
+	}
+}
+
+void UFastGeoWorldSubsystem::OnRemoveLevelFromWorldExtension(ULevel* Level, const bool bWaitForCompletion, bool& bOutHasCompleted)
+{
+	if (UFastGeoContainer* FastGeo = Level->GetAssetUserData<UFastGeoContainer>())
+	{
+		TRACE_CPUPROFILER_EVENT_SCOPE(UFastGeoWorldSubsystem::OnRemoveLevelFromWorldExtension);
+
+		if (FastGeo->HasAnyPendingDestroyTasks())
+		{
+			FastGeo->Tick(bWaitForCompletion);
+		}
+
+		if (FastGeo->HasAnyPendingDestroyTasks())
+		{
+			bOutHasCompleted = false;
+		}
+	}
+}
+
+void UFastGeoWorldSubsystem::OnUpdateLevelStreaming()
+{
+	UWorld* World = GetWorld();
+	if (World->GetShouldForceUnloadStreamingLevels())
+	{
+		TRACE_CPUPROFILER_EVENT_SCOPE(UFastGeoWorldSubsystem::OnUpdateLevelStreaming);
+
+		const bool bWaitForCompletion = true;
+		FFastGeoStreamingManager::ProcessAsyncRenderStateJobs(bWaitForCompletion);
+		if (FPhysScene* PhysScene = World->GetPhysicsScene())
+		{
+			PhysScene->ProcessAsyncPhysicsStateJobs(bWaitForCompletion);
+		}
+	}
+}
+
+void UFastGeoWorldSubsystem::AddToComponentsPendingRecreate(FFastGeoPrimitiveComponent* InComponentPendingRecreate)
+{
+	ComponentsPendingRecreate.Add(InComponentPendingRecreate);
+}
+
+void UFastGeoWorldSubsystem::RemoveFromComponentsPendingRecreate(FFastGeoPrimitiveComponent* InComponentPendingRecreate)
+{
+	ComponentsPendingRecreate.Remove(InComponentPendingRecreate);
+}
+
+void UFastGeoWorldSubsystem::ProcessPendingRecreate()
+{
+	for (FWeakFastGeoComponent WeakComponentPendingRecreate : ComponentsPendingRecreate)
+	{
+		if (FFastGeoPrimitiveComponent* ComponentPendingRecreate = WeakComponentPendingRecreate.Get<FFastGeoPrimitiveComponent>())
+		{
+			// Skip component if its container was unregistered before we process it
+			if (ComponentPendingRecreate->GetOwnerContainer()->IsRegistered())
+			{
+				ComponentPendingRecreate->DestroyRenderState(nullptr);
+				ComponentPendingRecreate->CreateRenderState(nullptr);
+			}
+		}
+	}
+
+	ComponentsPendingRecreate.Reset();
+}
+
+void UFastGeoWorldSubsystem::ForEachHLODObjectInCell(const UWorldPartitionRuntimeCell* InCell, TFunction<void(IWorldPartitionHLODObject*)> InFunc)
+{
+	check(InCell);
+	check(InCell->GetLevel());
+
+	if (UFastGeoContainer* FastGeo = InCell->GetLevel()->GetAssetUserData<UFastGeoContainer>())
+	{
+		// Iterate over clusters in the container, and call InFunc on all HLOD objects
+		FastGeo->ForEachComponentCluster<FFastGeoHLOD>([&InFunc](FFastGeoHLOD& HLOD)
+		{
+			InFunc(&HLOD);
+		});
+	}
+}
+
+#if DO_CHECK
+void UFastGeoWorldSubsystem::CheckNoPendingTasks(ULevel* Level, UWorld* World)
+{
+	if (World != GetWorld())
+	{
+		return;
+	}
+
+	if (UFastGeoContainer* FastGeo = Level->GetAssetUserData<UFastGeoContainer>())
+	{
+		check(!FastGeo->HasAnyPendingTasks());
+	}
+}
+
+void UFastGeoWorldSubsystem::OnLevelAddedToWorld(ULevel* Level, UWorld* World)
+{
+	check(Level);
+	CheckNoPendingTasks(Level, World);
+}
+
+void UFastGeoWorldSubsystem::OnLevelRemovedFromWorld(ULevel* Level, UWorld* World)
+{
+	// Null means every sublevel is being removed
+	if (!Level)
+	{
+		check(World);
+		for (ULevel* SubLevel : World->GetLevels())
+		{
+			CheckNoPendingTasks(SubLevel, World);
+		}
+	}
+	else
+	{
+		CheckNoPendingTasks(Level, World);
+	}
+}
+#endif
+
+#undef LOCTEXT_NAMESPACE
\ No newline at end of file
diff --git a/Engine/Plugins/Experimental/FastGeoStreaming/Source/FastGeoStreaming/Private/FastGeoWorldSubsystem.h b/Engine/Plugins/Experimental/FastGeoStreaming/Source/FastGeoStreaming/Private/FastGeoWorldSubsystem.h
new file mode 100644
--- /dev/null
+++ b/Engine/Plugins/Experimental/FastGeoStreaming/Source/FastGeoStreaming/Private/FastGeoWorldSubsystem.h
@@ -0,0 +1,60 @@
+// Copyright Epic Games, Inc. All Rights Reserved.
+
+#pragma once
+
+#include "Subsystems/WorldSubsystem.h"
+#include "FastGeoWeakElement.h"
+#include "FastGeoWorldSubsystem.generated.h"
+
+class FFastGeoPrimitiveComponent;
+class ULevelStreaming;
+class ULevel;
+class IWorldPartitionHLODObject;
+
+UCLASS()
+class FASTGEOSTREAMING_API UFastGeoWorldSubsystem : public UWorldSubsystem
+{
+	GENERATED_BODY()
+
+	UFastGeoWorldSubsystem();
+	~UFastGeoWorldSubsystem();
+
+public:
+	//~Begin WorldSubsystem interface
+	virtual void PostInitialize() override;
+	virtual void Deinitialize() override;
+	//~End WorldSubsystem interface
+
+	//~ Begin UWorldSubsystem Interface.
+	virtual bool DoesSupportWorldType(const EWorldType::Type WorldType) const override;
+	//~ End UWorldSubsystem Interface.
+
+	void AddToComponentsPendingRecreate(FFastGeoPrimitiveComponent* InComponentPendingRecreate);
+	void RemoveFromComponentsPendingRecreate(FFastGeoPrimitiveComponent* InComponentPendingRecreate);
+	void ProcessPendingRecreate();
+
+	static bool IsEnableDebugView();
+
+private:
+	void OnUpdateLevelStreaming();
+	void OnLevelStartedAddToWorld(UWorld* World, const ULevelStreaming* StreamingLevel, ULevel* Level);
+	void OnLevelStartedRemoveFromWorld(UWorld* World, const ULevelStreaming* StreamingLevel, ULevel* Level);
+	void OnLevelComponentsUpdated(UWorld* World, ULevel* Level);
+	void OnLevelComponentsCleared(UWorld* World, ULevel* Level);
+	void OnAddLevelToWorldExtension(ULevel* Level, const bool bWaitForCompletion, bool& bOutHasCompleted);
+	void OnRemoveLevelFromWorldExtension(ULevel* Level, const bool bWaitForCompletion, bool& bOutHasCompleted);
+#if DO_CHECK
+	void OnLevelAddedToWorld(ULevel* Level, UWorld* World);
+	void OnLevelRemovedFromWorld(ULevel* Level, UWorld* World);
+	void CheckNoPendingTasks(ULevel* Level, UWorld* World);
+#endif
+
+	void ForEachHLODObjectInCell(const UWorldPartitionRuntimeCell*, TFunction<void(IWorldPartitionHLODObject*)>);
+
+	FDelegateHandle Handle_OnLevelBeginAddToWorld;
+	FDelegateHandle Handle_OnLevelBeginRemoveFromWorld;
+	FDelegateHandle Handle_OnForEachHLODObjectInCell;
+
+	TArray<FWeakFastGeoComponent> ComponentsPendingRecreate;
+	static bool bEnableDebugView;
+};
diff --git a/Engine/Plugins/Experimental/FastGeoStreaming/Source/FastGeoStreaming/Private/IFastGeoElement.cpp b/Engine/Plugins/Experimental/FastGeoStreaming/Source/FastGeoStreaming/Private/IFastGeoElement.cpp
new file mode 100644
--- /dev/null
+++ b/Engine/Plugins/Experimental/FastGeoStreaming/Source/FastGeoStreaming/Private/IFastGeoElement.cpp
@@ -0,0 +1,6 @@
+// Copyright Epic Games, Inc. All Rights Reserved.
+
+#include "IFastGeoElement.h"
+
+uint32 FFastGeoElementType::NextUniqueID = 0;
+const FFastGeoElementType IFastGeoElement::Type;
\ No newline at end of file
diff --git a/Engine/Plugins/Experimental/FastGeoStreaming/Source/FastGeoStreaming/Private/IFastGeoElement.h b/Engine/Plugins/Experimental/FastGeoStreaming/Source/FastGeoStreaming/Private/IFastGeoElement.h
new file mode 100644
--- /dev/null
+++ b/Engine/Plugins/Experimental/FastGeoStreaming/Source/FastGeoStreaming/Private/IFastGeoElement.h
@@ -0,0 +1,57 @@
+// Copyright Epic Games, Inc. All Rights Reserved.
+
+#pragma once
+
+#include "FastGeoElementType.h"
+
+class IFastGeoElement
+{
+public:
+	/** Default constructor */
+	IFastGeoElement(FFastGeoElementType InType) : ElementType(InType) {}
+	virtual ~IFastGeoElement() {}
+
+	/** Attempt to cast this element to another type if it is of that type. Returns nullptr if it fails */
+	template <typename T>
+	T* CastTo()
+	{
+		return ElementType.IsA(T::Type) ? StaticCast<T*>(this) : nullptr;
+	}
+	/** Attempt to cast this element to another type if it is of that type. Returns nullptr if it fails */
+	template <typename T>
+	const T* CastTo() const
+	{
+		return ElementType.IsA(T::Type) ? StaticCast<const T*>(this) : nullptr;
+	}
+	/** Cast this element to another type and return a reference. Will assert if the cast fails. */
+	template <typename T>
+	T& CastToRef()
+	{
+		check(ElementType.IsA(T::Type));
+		return *StaticCast<T*>(this);
+	}
+	/** Cast this element to another type and return a reference. Will assert if the cast fails. */
+	template <typename T>
+	const T& CastToRef() const
+	{
+		check(ElementType.IsA(T::Type));
+		return *StaticCast<const T*>(this);
+	}
+	/** Returns true if this element is of the specified type */
+	template <typename T>
+	bool IsA() const
+	{
+		return ElementType.IsA(T::Type);
+	}
+
+	/** Returns the element type unique ID */
+	uint32 GetTypeID() { return ElementType.ID; }
+
+protected:
+	/** Static type identifier for the base class */
+	static const FFastGeoElementType Type;
+
+	/** Element type identifier */
+	FFastGeoElementType ElementType;
+};
+
diff --git a/Engine/Plugins/Experimental/FastGeoStreaming/Source/FastGeoStreaming/Public/FastGeoStreamingModule.h b/Engine/Plugins/Experimental/FastGeoStreaming/Source/FastGeoStreaming/Public/FastGeoStreamingModule.h
new file mode 100644
--- /dev/null
+++ b/Engine/Plugins/Experimental/FastGeoStreaming/Source/FastGeoStreaming/Public/FastGeoStreamingModule.h
@@ -0,0 +1,28 @@
+// Copyright Epic Games, Inc. All Rights Reserved.
+
+#pragma once
+
+#include "Delegates/Delegate.h"
+#include "Modules/ModuleInterface.h"
+
+class FFastGeoPrimitiveComponent;
+class UWorld;
+
+class FFastGeoStreamingModule : public IModuleInterface
+{
+public:
+	static FFastGeoStreamingModule& Get();
+
+	/** IModuleInterface implementation */
+	virtual void StartupModule() override;
+	virtual void ShutdownModule() override;
+
+	void AddToComponentsPendingRecreate(FFastGeoPrimitiveComponent* InComponentPendingRecreate);
+	void RemoveFromComponentsPendingRecreate(FFastGeoPrimitiveComponent* InComponentPendingRecreate);
+
+private:
+	void OnWorldPreSendAllEndOfFrameUpdates(UWorld* InWorld);
+
+private:
+	FDelegateHandle Handle_OnWorldPreSendAllEndOfFrameUpdates;
+};
