CIS fix

diff --git a/Engine/Plugins/Experimental/FastGeoStreaming/Source/FastGeoStreaming/Private/FastGeoComponentCluster.h b/Engine/Plugins/Experimental/FastGeoStreaming/Source/FastGeoStreaming/Private/FastGeoComponentCluster.h
--- a/Engine/Plugins/Experimental/FastGeoStreaming/Source/FastGeoStreaming/Private/FastGeoComponentCluster.h
+++ b/Engine/Plugins/Experimental/FastGeoStreaming/Source/FastGeoStreaming/Private/FastGeoComponentCluster.h
@@ -91,7 +91,7 @@ protected:
 		auto ForEachArray = [&InFunc](auto& InArray)
 		{
 			typedef typename TDecay<decltype(InArray)>::Type ArrayType;
-			using TArrayComponent = ArrayType::ElementType;
+			using TArrayComponent = typename ArrayType::ElementType;
 
 			if constexpr (std::is_same<TComponent, FFastGeoComponent>::value)
 			{
@@ -120,7 +120,7 @@ protected:
 		auto ForEachArray = [&InFunc](auto& InArray) -> bool
 		{
 			typedef typename TDecay<decltype(InArray)>::Type ArrayType;
-			using TArrayComponent = ArrayType::ElementType;
+			using TArrayComponent = typename ArrayType::ElementType;
 
 			if constexpr (std::is_same<TComponent, FFastGeoComponent>::value)
 			{
diff --git a/Engine/Plugins/Experimental/FastGeoStreaming/Source/FastGeoStreaming/Private/FastGeoContainer.h b/Engine/Plugins/Experimental/FastGeoStreaming/Source/FastGeoStreaming/Private/FastGeoContainer.h
--- a/Engine/Plugins/Experimental/FastGeoStreaming/Source/FastGeoStreaming/Private/FastGeoContainer.h
+++ b/Engine/Plugins/Experimental/FastGeoStreaming/Source/FastGeoStreaming/Private/FastGeoContainer.h
@@ -132,7 +132,7 @@ private:
 		auto ForEachArray = [&InFunc](auto& InArray)
 		{
 			typedef typename TDecay<decltype(InArray)>::Type ArrayType;
-			using TArrayComponentCluster = ArrayType::ElementType;
+			using TArrayComponentCluster = typename ArrayType::ElementType;
 
 			if constexpr (std::is_same<TComponentCluster, FFastGeoComponentCluster>::value)
 			{
@@ -161,7 +161,7 @@ private:
 		auto ForEachArray = [&InFunc](auto& InArray) -> bool
 		{
 			typedef typename TDecay<decltype(InArray)>::Type ArrayType;
-			using TArrayComponentCluster = ArrayType::ElementType;
+			using TArrayComponentCluster = typename ArrayType::ElementType;
 
 			if constexpr (std::is_same<TComponentCluster, FFastGeoComponentCluster>::value)
 			{
