FastGeo / FastGeoWorldPartitionRuntimeCellTransformer:
- Added AllowedExactActorClasses, AllowedComponentClasses, AllowedExactComponentClasses, DisallowedExactActorClasses, DisallowedComponentClasses, DisallowedExactComponentClasses, IgnoredRemainingExactComponentClasses to better control what to include/exclude during the FastGeo transformation.
- Automatic exclusion of USplineMeshComponent,InteractiveFoliageComponent,LandscapeNaniteComponent,LandscapeMeshProxyComponent,WaterBodyMeshComponent and AWaterBody.

diff --git a/Engine/Plugins/Experimental/FastGeoStreaming/Source/FastGeoStreaming/Private/FastGeoWorldPartitionRuntimeCellTransformer.cpp b/Engine/Plugins/Experimental/FastGeoStreaming/Source/FastGeoStreaming/Private/FastGeoWorldPartitionRuntimeCellTransformer.cpp
--- a/Engine/Plugins/Experimental/FastGeoStreaming/Source/FastGeoStreaming/Private/FastGeoWorldPartitionRuntimeCellTransformer.cpp
+++ b/Engine/Plugins/Experimental/FastGeoStreaming/Source/FastGeoStreaming/Private/FastGeoWorldPartitionRuntimeCellTransformer.cpp
@@ -9,6 +9,7 @@
 #include "UObject/Package.h"
 #include "Components/StaticMeshComponent.h"
 #include "Components/InstancedStaticMeshComponent.h"
+#include "Components/SplineMeshComponent.h"
 #include "FastGeoContainer.h"
 #include "FastGeoComponentCluster.h"
 #include "FastGeoComponent.h"
@@ -29,15 +30,44 @@
 
 #define LOCTEXT_NAMESPACE "WorldPartition"
 
+#if WITH_EDITORONLY_DATA
+namespace FastGeo
+{
+	struct FClassHelper
+	{
+		template <typename T>
+		static void AddClass(const TCHAR* Name, TArray<TSubclassOf<T>>& OutClassArray)
+		{
+			if (UClass* Class = FindObject<UClass>(nullptr, Name, true); ensure(Class))
+			{
+				OutClassArray.Add(Class);
+			}
+		}
+	};
+}
+#endif
+
 UFastGeoWorldPartitionRuntimeCellTransformer::UFastGeoWorldPartitionRuntimeCellTransformer(const FObjectInitializer& ObjectInitializer)
 	: Super(ObjectInitializer)
 {
 #if WITH_EDITORONLY_DATA
-	AllowedClasses.Add(APackedLevelActor::StaticClass());
-	AllowedClasses.Add(AStaticMeshActor::StaticClass());
-	AllowedClasses.Add(AWorldPartitionHLOD::StaticClass());
-	IgnoredComponentClasses.Add(ULevelInstanceComponent::StaticClass());
-	IgnoredComponentClasses.Add(UActorTextureStreamingBuildDataComponent::StaticClass());
+	// Allowed actor/component classes
+	AllowedActorClasses.Add(APackedLevelActor::StaticClass());
+	AllowedActorClasses.Add(AStaticMeshActor::StaticClass());
+	AllowedActorClasses.Add(AWorldPartitionHLOD::StaticClass());
+	AllowedComponentClasses.Add(UStaticMeshComponent::StaticClass());
+
+	// Disallowed actor/component classes
+	FastGeo::FClassHelper::AddClass(TEXT("/Script/Water.WaterBody"), DisallowedActorClasses);
+	FastGeo::FClassHelper::AddClass(TEXT("/Script/Foliage.InteractiveFoliageComponent"), DisallowedComponentClasses);
+	FastGeo::FClassHelper::AddClass(TEXT("/Script/Landscape.LandscapeNaniteComponent"), DisallowedComponentClasses);
+	FastGeo::FClassHelper::AddClass(TEXT("/Script/Landscape.LandscapeMeshProxyComponent"), DisallowedComponentClasses);
+	FastGeo::FClassHelper::AddClass(TEXT("/Script/Water.WaterBodyMeshComponent"), DisallowedComponentClasses);
+	DisallowedComponentClasses.Add(USplineMeshComponent::StaticClass());
+
+	// Ignored component classes
+	IgnoredRemainingComponentClasses.Add(ULevelInstanceComponent::StaticClass());
+	IgnoredRemainingComponentClasses.Add(UActorTextureStreamingBuildDataComponent::StaticClass());
 #endif
 }
 
@@ -86,18 +116,57 @@ void UFastGeoWorldPartitionRuntimeCellTransformer::Transform(ULevel* InLevel)
 
 	check(InLevel);
 
+	auto IsAllowedComponentClass = [this](UActorComponent* Component)
+	{
+		UClass* ComponentClass = Component->GetClass();
+		for (TSubclassOf<UActorComponent> DisallowedComponentClass : DisallowedComponentClasses)
+		{
+			if (ComponentClass->IsChildOf(DisallowedComponentClass))
+			{
+				return false;
+			}
+		}
+
+		for (TSubclassOf<UActorComponent> DisallowedComponentClass : DisallowedExactComponentClasses)
+		{
+			if (ComponentClass == DisallowedComponentClass)
+			{
+				return false;
+			}
+		}
+
+		for (TSubclassOf<UActorComponent> AllowedComponentClass : AllowedComponentClasses)
+		{
+			if (ComponentClass->IsChildOf(AllowedComponentClass))
+			{
+				return true;
+			}
+		}
+
+		for (TSubclassOf<UActorComponent> AllowedComponentClass : AllowedExactComponentClasses)
+		{
+			if (ComponentClass == AllowedComponentClass)
+			{
+				return true;
+			}
+		}
+
+		return false;
+	};
+
 	TMap<TObjectPtr<AActor>*, TArray<UStaticMeshComponent*>> ActorComponents;
 	for (TObjectPtr<AActor>& Actor : InLevel->Actors)
 	{
 		if (IsValid(Actor) && CanTransformActor(Actor))
 		{
 			// Gather components
-			Actor->ForEachComponent<UStaticMeshComponent>(false, [&Actor, &ActorComponents](UStaticMeshComponent* StaticMeshComponent)
+			Actor->ForEachComponent<UStaticMeshComponent>(false, [&Actor, &ActorComponents, &IsAllowedComponentClass](UStaticMeshComponent* StaticMeshComponent)
 			{
 				if (!StaticMeshComponent->IsEditorOnly() &&
 					StaticMeshComponent->GetStaticMesh() &&
 					!StaticMeshComponent->GetLODParentPrimitive() &&
-					(StaticMeshComponent->Mobility == EComponentMobility::Static))
+					(StaticMeshComponent->Mobility == EComponentMobility::Static) &&
+					IsAllowedComponentClass(StaticMeshComponent))
 				{
 					// Make sure BodyInstance CollisionEnabled is updated first
 					StaticMeshComponent->UpdateCollisionFromStaticMesh();
@@ -125,8 +194,8 @@ void UFastGeoWorldPartitionRuntimeCellTransformer::Transform(ULevel* InLevel)
 	};
 
 	// Walk the component class hierarchy and look for a fast geo mapping
-	auto GetFastGeoComponentType = [&FastGeoComponentTypeMapping] (TSubclassOf<UPrimitiveComponent> InClass) -> FFastGeoElementType
-	{		
+	auto GetFastGeoComponentType = [&FastGeoComponentTypeMapping](TSubclassOf<UPrimitiveComponent> InClass) -> FFastGeoElementType
+	{
 		for (UClass* Class = InClass; Class; Class = Class->GetSuperClass())
 		{
 			if (const FFastGeoElementType* Found = FastGeoComponentTypeMapping.Find(Class))
@@ -225,9 +294,17 @@ bool UFastGeoWorldPartitionRuntimeCellTransformer::CanTransformActor(AActor* InA
 	}
 
 	UClass* ActorClass = InActor->GetClass();
-	for (TSubclassOf<AActor> DisallowedClass : DisallowedClasses)
+	for (TSubclassOf<AActor> DisallowedActorClass : DisallowedActorClasses)
+	{
+		if (ActorClass->IsChildOf(DisallowedActorClass))
+		{
+			return false;
+		}
+	}
+
+	for (TSubclassOf<AActor> DisallowedActorClass : DisallowedExactActorClasses)
 	{
-		if (ActorClass == DisallowedClass)
+		if (ActorClass == DisallowedActorClass)
 		{
 			return false;
 		}
@@ -238,9 +315,17 @@ bool UFastGeoWorldPartitionRuntimeCellTransformer::CanTransformActor(AActor* InA
 		return false;
 	}
 
-	for (TSubclassOf<AActor> AllowedClass : AllowedClasses)
+	for (TSubclassOf<AActor> AllowedActorClass : AllowedActorClasses)
+	{
+		if (ActorClass->IsChildOf(AllowedActorClass))
+		{
+			return true;
+		}
+	}
+
+	for (TSubclassOf<AActor> AllowedActorClass : AllowedExactActorClasses)
 	{
-		if (ActorClass->IsChildOf(AllowedClass))
+		if (ActorClass->IsChildOf(AllowedActorClass))
 		{
 			return true;
 		}
@@ -277,13 +362,22 @@ bool UFastGeoWorldPartitionRuntimeCellTransformer::CanRemoveActor(AActor* InActo
 		}
 
 		UClass* ComponentClass = Component->GetClass();
-		for (const TSubclassOf<UActorComponent>& IgnoredComponentClass : IgnoredComponentClasses)
+		for (const TSubclassOf<UActorComponent>& IgnoredComponentClass : IgnoredRemainingComponentClasses)
 		{
 			if (ComponentClass->IsChildOf(IgnoredComponentClass))
 			{
 				return true;
 			}
 		}
+
+		for (const TSubclassOf<UActorComponent>& IgnoredComponentClass : IgnoredRemainingExactComponentClasses)
+		{
+			if (ComponentClass == IgnoredComponentClass)
+			{
+				return true;
+			}
+		}
+
 		return false;
 	};
 
@@ -295,7 +389,7 @@ bool UFastGeoWorldPartitionRuntimeCellTransformer::CanRemoveActor(AActor* InActo
 			NumRuntimeComponents++;
 		}
 	}
-		
+
 	if (!NumRuntimeComponents)
 	{
 		return true;
@@ -316,4 +410,4 @@ bool UFastGeoWorldPartitionRuntimeCellTransformer::CanRemoveActor(AActor* InActo
 }
 #endif
 
-#undef LOCTEXT_NAMESPACE
+#undef LOCTEXT_NAMESPACE
\ No newline at end of file
diff --git a/Engine/Plugins/Experimental/FastGeoStreaming/Source/FastGeoStreaming/Private/FastGeoWorldPartitionRuntimeCellTransformer.h b/Engine/Plugins/Experimental/FastGeoStreaming/Source/FastGeoStreaming/Private/FastGeoWorldPartitionRuntimeCellTransformer.h
--- a/Engine/Plugins/Experimental/FastGeoStreaming/Source/FastGeoStreaming/Private/FastGeoWorldPartitionRuntimeCellTransformer.h
+++ b/Engine/Plugins/Experimental/FastGeoStreaming/Source/FastGeoStreaming/Private/FastGeoWorldPartitionRuntimeCellTransformer.h
@@ -23,15 +23,44 @@ private:
 #endif
 
 #if WITH_EDITORONLY_DATA
-	/** Allowed classes (recursive) to convert to instances */
+	/** Allowed actor classes (recursive) to convert to FastGeo */
 	UPROPERTY(EditAnywhere, Category = FastGeo)
-	TArray<TSubclassOf<AActor>> AllowedClasses;
+	TArray<TSubclassOf<AActor>> AllowedActorClasses;
 
-	/** Disallowed classes (non-recursive) to convert to instances */
+	/** Allowed actor classes (exact) to convert to FastGeo */
 	UPROPERTY(EditAnywhere, Category = FastGeo)
-	TArray<TSubclassOf<AActor>> DisallowedClasses;
+	TArray<TSubclassOf<AActor>> AllowedExactActorClasses;
 
+	/** Allowed component classes (recursive) to convert to FastGeo */
 	UPROPERTY(EditAnywhere, Category = FastGeo)
-	TArray<TSubclassOf<UActorComponent>> IgnoredComponentClasses;
+	TArray<TSubclassOf<UActorComponent>> AllowedComponentClasses;
+
+	/** Allowed component classes (exact) to convert to FastGeo */
+	UPROPERTY(EditAnywhere, Category = FastGeo)
+	TArray<TSubclassOf<UActorComponent>> AllowedExactComponentClasses;
+
+	/** Disallowed actor classes (recursive) to convert to FastGeo */
+	UPROPERTY(EditAnywhere, Category = FastGeo)
+	TArray<TSubclassOf<AActor>> DisallowedActorClasses;
+
+	/** Disallowed actor classes (exact) to convert to FastGeo */
+	UPROPERTY(EditAnywhere, Category = FastGeo)
+	TArray<TSubclassOf<AActor>> DisallowedExactActorClasses;
+
+	/** Disallowed component classes (recursive) to convert to FastGeo */
+	UPROPERTY(EditAnywhere, Category = FastGeo)
+	TArray<TSubclassOf<UActorComponent>> DisallowedComponentClasses;
+
+	/** Disallowed component classes (exact) to convert to FastGeo */
+	UPROPERTY(EditAnywhere, Category = FastGeo)
+	TArray<TSubclassOf<UActorComponent>> DisallowedExactComponentClasses;
+
+	/** Component classes (recursive) to ignore when deciding to destroy a converted actor to FastGeo */
+	UPROPERTY(EditAnywhere, Category = FastGeo)
+	TArray<TSubclassOf<UActorComponent>> IgnoredRemainingComponentClasses;
+
+	/** Component classes (exact) to ignore when deciding to destroy a converted actor to FastGeo */
+	UPROPERTY(EditAnywhere, Category = FastGeo)
+	TArray<TSubclassOf<UActorComponent>> IgnoredRemainingExactComponentClasses;
 #endif
 };
\ No newline at end of file
