FastGeo: Fixed invalid check in FastGeoAsyncrenderStateJobQueue

diff --git a/Engine/Plugins/Experimental/FastGeoStreaming/Source/FastGeoStreaming/Private/FastGeoAsyncRenderStateJobQueue.cpp b/Engine/Plugins/Experimental/FastGeoStreaming/Source/FastGeoStreaming/Private/FastGeoAsyncRenderStateJobQueue.cpp
--- a/Engine/Plugins/Experimental/FastGeoStreaming/Source/FastGeoStreaming/Private/FastGeoAsyncRenderStateJobQueue.cpp
+++ b/Engine/Plugins/Experimental/FastGeoStreaming/Source/FastGeoStreaming/Private/FastGeoAsyncRenderStateJobQueue.cpp
@@ -65,12 +65,7 @@ bool FFastGeoAsyncRenderStateJobQueue::AreAsyncTasksExecuted() const
 
 bool FFastGeoAsyncRenderStateJobQueue::IsCompleted() const
 {
-	if (!PendingJobs.IsValid() && PipedJobs.IsEmpty() && PipeTasks.IsEmpty())
-	{
-		check(AreAsyncTasksExecuted());
-		return true;
-	}
-	return false;
+	return AreAsyncTasksExecuted() && PipedJobs.IsEmpty() && PipeTasks.IsEmpty();
 }
 
 void FFastGeoAsyncRenderStateJobQueue::Launch()
