- FastGeo Cell Transformer doesn't allow convert a BP actor with logic (i.e. if it has more functionality than its construction script).
- FastGeo Cell Transformer also doesn't convert partially a BP actor (rerun CS will be called during component registration in PIE and during the cooking).

diff --git a/Engine/Plugins/Experimental/FastGeoStreaming/Source/FastGeoStreaming/Private/FastGeoWorldPartitionRuntimeCellTransformer.cpp b/Engine/Plugins/Experimental/FastGeoStreaming/Source/FastGeoStreaming/Private/FastGeoWorldPartitionRuntimeCellTransformer.cpp
--- a/Engine/Plugins/Experimental/FastGeoStreaming/Source/FastGeoStreaming/Private/FastGeoWorldPartitionRuntimeCellTransformer.cpp
+++ b/Engine/Plugins/Experimental/FastGeoStreaming/Source/FastGeoStreaming/Private/FastGeoWorldPartitionRuntimeCellTransformer.cpp
@@ -6,6 +6,9 @@
 #include "CoreMinimal.h"
 #include "GameFramework/Actor.h"
 #include "Engine/Level.h"
+#include "Engine/Blueprint.h"
+#include "Engine/BlueprintGeneratedClass.h"
+#include "EdGraph/EdGraph.h"
 #include "UObject/Package.h"
 #include "Components/StaticMeshComponent.h"
 #include "Components/InstancedStaticMeshComponent.h"
@@ -154,7 +157,7 @@ void UFastGeoWorldPartitionRuntimeCellTransformer::Transform(ULevel* InLevel)
 		return false;
 	};
 
-	TMap<TObjectPtr<AActor>*, TArray<UStaticMeshComponent*>> ActorComponents;
+	TMap<TObjectPtr<AActor>*, TSet<UActorComponent*>> ActorComponents;
 	for (TObjectPtr<AActor>& Actor : InLevel->Actors)
 	{
 		if (IsValid(Actor) && CanTransformActor(Actor))
@@ -220,6 +223,15 @@ void UFastGeoWorldPartitionRuntimeCellTransformer::Transform(ULevel* InLevel)
 		{
 			check(!Components.IsEmpty());
 
+			// Can't convert partially a BP actors
+			// Rerun CS will be called in PIE when registering the component and
+			// also called when registering components during cook/save of the level.
+			const bool bCanRemoveActor = CanRemoveActor(*Actor, Components);
+			if (UBlueprint::GetBlueprintFromClass(Actor->GetClass()) && !bCanRemoveActor)
+			{
+				continue;
+			}
+
 			FFastGeoComponentCluster* CurrentComponentCluster = LevelComponentCluster.Get();
 			TUniquePtr<FFastGeoHLOD> FastGeoHLOD;
 			if (AWorldPartitionHLOD* HLODActor = Cast<AWorldPartitionHLOD>(*Actor))
@@ -230,8 +242,9 @@ void UFastGeoWorldPartitionRuntimeCellTransformer::Transform(ULevel* InLevel)
 				CurrentComponentCluster = FastGeoHLOD.Get();
 			}
 
-			for (UStaticMeshComponent* StaticMeshComponent : Components)
+			for (UActorComponent* Component : Components)
 			{
+				UStaticMeshComponent* StaticMeshComponent = CastChecked<UStaticMeshComponent>(Component);
 				FFastGeoElementType FastGeoComponentType = GetFastGeoComponentType(StaticMeshComponent->GetClass());
 				FFastGeoComponent& FastGeoComponent = CurrentComponentCluster->AddComponent(FastGeoComponentType);
 				FastGeoComponent.InitializeFromComponent(StaticMeshComponent);
@@ -241,7 +254,7 @@ void UFastGeoWorldPartitionRuntimeCellTransformer::Transform(ULevel* InLevel)
 				StaticMeshComponent->MarkAsGarbage();
 			}
 
-			if (CanRemoveActor(*Actor))
+			if (bCanRemoveActor)
 			{
 				*Actor = nullptr;
 			}
@@ -266,6 +279,69 @@ void UFastGeoWorldPartitionRuntimeCellTransformer::Transform(ULevel* InLevel)
 	InLevel->Actors.Remove(nullptr);
 }
 
+bool UFastGeoWorldPartitionRuntimeCellTransformer::IsBlueprintActorWithLogic(AActor* InActor) const
+{
+	const FName FN_UserConstructionScript(TEXT("UserConstructionScript"));
+
+	UBlueprint* Blueprint = UBlueprint::GetBlueprintFromClass(InActor->GetClass());
+	if (!Blueprint || !Blueprint->GeneratedClass)
+	{
+		return false;
+	}
+
+	check(Blueprint->ParentClass && Blueprint->ParentClass->IsChildOf(AActor::StaticClass()));
+
+	UBlueprintGeneratedClass* BPClass = Cast<UBlueprintGeneratedClass>(Blueprint->GeneratedClass);
+	if (!BPClass)
+	{
+		return false;
+	}
+
+	if (Blueprint->DelegateSignatureGraphs.Num() > 0)
+	{
+		return true;
+	}
+
+	if (Blueprint->ImplementedInterfaces.Num() > 0)
+	{
+		return true;
+	}
+
+	// Check if no extra functions, other than the user construction script (only AActor and subclasses of AActor have)
+	if ((Blueprint->FunctionGraphs.Num() > 1) || (Blueprint->MacroGraphs.Num() > 0))
+	{
+		return true;
+	}
+
+	check(Blueprint->FunctionGraphs.Num() == 0 || Blueprint->FunctionGraphs[0]->GetFName() == FN_UserConstructionScript);
+
+	// Check if the generated class has overridden any functions dynamically
+	for (TFieldIterator<UFunction> It(BPClass, EFieldIteratorFlags::IncludeSuper); It; ++It)
+	{
+		UFunction* Function = *It;
+
+		// Ignore functions from native C++ classes (i.e., inherited but not overridden in BP)
+		if (Function->GetOwnerClass() == BPClass && Function->GetFName() != FN_UserConstructionScript)
+		{
+			return true; // Found an overridden function
+		}
+	}
+
+	// If there is an enabled node in the event graph, the Blueprint is not data only
+	for (UEdGraph* EventGraph : Blueprint->UbergraphPages)
+	{
+		for (UEdGraphNode* GraphNode : EventGraph->Nodes)
+		{
+			if (GraphNode && (GraphNode->GetDesiredEnabledState() != ENodeEnabledState::Disabled))
+			{
+				return true;
+			}
+		}
+	}
+
+	return false;
+}
+
 bool UFastGeoWorldPartitionRuntimeCellTransformer::CanTransformActor(AActor* InActor) const
 {
 	if (InActor->GetIsReplicated())
@@ -315,6 +391,11 @@ bool UFastGeoWorldPartitionRuntimeCellTransformer::CanTransformActor(AActor* InA
 		return false;
 	}
 
+	if (IsBlueprintActorWithLogic(InActor))
+	{
+		return false;
+	}
+
 	for (TSubclassOf<AActor> AllowedActorClass : AllowedActorClasses)
 	{
 		if (ActorClass->IsChildOf(AllowedActorClass))
@@ -339,10 +420,15 @@ bool UFastGeoWorldPartitionRuntimeCellTransformer::CanTransformActor(AActor* InA
 	return false;
 }
 
-bool UFastGeoWorldPartitionRuntimeCellTransformer::CanRemoveActor(AActor* InActor) const
+bool UFastGeoWorldPartitionRuntimeCellTransformer::CanRemoveActor(AActor* InActor, TSet<UActorComponent*>& InIgnoredComponents) const
 {
-	auto IsComponentIgnored = [this](UActorComponent* Component)
+	auto IsComponentIgnored = [this, &InIgnoredComponents](UActorComponent* Component)
 	{
+		if (InIgnoredComponents.Contains(Component))
+		{
+			return true;
+		}
+
 		if (Component->IsEditorOnly())
 		{
 			return true;
diff --git a/Engine/Plugins/Experimental/FastGeoStreaming/Source/FastGeoStreaming/Private/FastGeoWorldPartitionRuntimeCellTransformer.h b/Engine/Plugins/Experimental/FastGeoStreaming/Source/FastGeoStreaming/Private/FastGeoWorldPartitionRuntimeCellTransformer.h
--- a/Engine/Plugins/Experimental/FastGeoStreaming/Source/FastGeoStreaming/Private/FastGeoWorldPartitionRuntimeCellTransformer.h
+++ b/Engine/Plugins/Experimental/FastGeoStreaming/Source/FastGeoStreaming/Private/FastGeoWorldPartitionRuntimeCellTransformer.h
@@ -18,8 +18,9 @@ public:
 	virtual void Transform(ULevel* InLevel) override;
 
 private:
+	bool IsBlueprintActorWithLogic(AActor* InActor) const;
 	bool CanTransformActor(AActor* InActor) const;
-	bool CanRemoveActor(AActor* InActor) const;
+	bool CanRemoveActor(AActor* InActor, TSet<UActorComponent*>& InIgnoredComponents) const;
 #endif
 
 #if WITH_EDITORONLY_DATA
