FastGeo: Dynamic NavMesh update support

diff --git a/Engine/Plugins/Experimental/FastGeoStreaming/Source/FastGeoStreaming/FastGeoStreaming.Build.cs b/Engine/Plugins/Experimental/FastGeoStreaming/Source/FastGeoStreaming/FastGeoStreaming.Build.cs
--- a/Engine/Plugins/Experimental/FastGeoStreaming/Source/FastGeoStreaming/FastGeoStreaming.Build.cs
+++ b/Engine/Plugins/Experimental/FastGeoStreaming/Source/FastGeoStreaming/FastGeoStreaming.Build.cs
@@ -12,6 +12,7 @@ namespace UnrealBuildTool.Rules
 					"Core",
 					"CoreUObject",
 					"Engine",
+					"NavigationSystem",
 					"PhysicsCore",
 					"RenderCore",
 					"RHI"
diff --git a/Engine/Plugins/Experimental/FastGeoStreaming/Source/FastGeoStreaming/Private/FastGeoContainer.cpp b/Engine/Plugins/Experimental/FastGeoStreaming/Source/FastGeoStreaming/Private/FastGeoContainer.cpp
--- a/Engine/Plugins/Experimental/FastGeoStreaming/Source/FastGeoStreaming/Private/FastGeoContainer.cpp
+++ b/Engine/Plugins/Experimental/FastGeoStreaming/Source/FastGeoStreaming/Private/FastGeoContainer.cpp
@@ -14,6 +14,8 @@
 #include "Physics/Experimental/PhysScene_Chaos.h"
 #include "WorldPartition/WorldPartitionHelpers.h"
 #include "WorldPartition/WorldPartition.h"
+#include "NavigationSystem.h"
+#include "AI/Navigation/NavigationElement.h"
 #include "PSOPrecache.h"
 
 class FFastGeoAssetRemapArchive : public FArchiveProxy
@@ -101,8 +103,11 @@ void UFastGeoContainer::Register()
 				});
 			}
 		});
+
 		FFastGeoStreamingManager::PushAsyncCreateRenderStateJob(this);
 		FFastGeoStreamingManager::PushAsyncCreatePhysicsStateJobs(this);
+		RegisterToNavigationSystem();
+
 		bIsRegistered = true;
 	}
 
@@ -134,6 +139,8 @@ void UFastGeoContainer::Unregister()
 		{
 			ComponentCluster.OnUnregister();
 		});
+
+		UnregisterFromNavigationSystem();
 		FFastGeoStreamingManager::PushAsyncDestroyRenderStateJob(this);
 		FFastGeoStreamingManager::PushAsyncDestroyPhysicsStateJobs(this);
 		bIsRegistered = false;
@@ -196,6 +203,62 @@ UWorld* UFastGeoContainer::GetWorld() const
 	return GetLevel()->GetWorld();
 }
 
+void UFastGeoContainer::RegisterToNavigationSystem()
+{
+	UWorld* World = GetWorld();
+	if (!FNavigationSystem::SupportsDynamicChanges(World))
+	{
+		return;
+	}
+	
+	TRACE_CPUPROFILER_EVENT_SCOPE(UFastGeoContainer::RegisterToNavigationSystem);
+	check(NavigationElementHandles.IsEmpty());
+
+	ForEachComponentCluster([this, World](FFastGeoComponentCluster& ComponentCluster)
+	{
+		ComponentCluster.ForEachComponent<FFastGeoPrimitiveComponent>([this, World](FFastGeoPrimitiveComponent& Component)
+		{
+			if (Component.IsNavigationRelevant())
+			{
+				FNavigationElement Element(*this, reinterpret_cast<const uint64>(&Component));
+				Element.SetDirtyAreaOnRegistration(!Component.ShouldSkipNavigationDirtyAreaOnAddOrRemove());
+				Element.SetBounds(Component.GetNavigationBounds());
+				Element.SetBodySetup(Component.GetBodySetup());
+				Element.SetTransform(Component.GetTransform());
+				Element.SetGeometryExportType(Component.HasCustomNavigableGeometry());
+				Element.NavigationDataExportDelegate.BindWeakLambda(this, [&Component](const FNavigationElement& NavigationElement, FNavigationRelevantData& OutNavigationRelevantData)
+				{
+					Component.GetNavigationData(OutNavigationRelevantData);
+				});
+				Element.CustomGeometryExportDelegate.BindWeakLambda(this, [&Component](const FNavigationElement& NavigationElement, FNavigableGeometryExport& OutGeometry, bool& bOutShouldExportDefaultGeometry)
+				{
+					bOutShouldExportDefaultGeometry = Component.DoCustomNavigableGeometryExport(OutGeometry);
+				});
+				FNavigationElementHandle Handle = FNavigationSystem::AddNavigationElement(World, MoveTemp(Element));
+				if (ensure(Handle.IsValid()))
+				{
+					NavigationElementHandles.Add(Handle);
+				}
+			}
+		});
+	});
+}
+
+void UFastGeoContainer::UnregisterFromNavigationSystem()
+{
+	if (!NavigationElementHandles.IsEmpty())
+	{
+		TRACE_CPUPROFILER_EVENT_SCOPE(UFastGeoContainer::UnregisterFromNavigationSystem);
+
+		for (FNavigationElementHandle& ElementHandle : NavigationElementHandles)
+		{
+			FNavigationSystem::RemoveNavigationElement(GetWorld(), ElementHandle);
+		}
+
+		NavigationElementHandles.Empty();
+	}
+}
+
 void UFastGeoContainer::OnCreateRenderStateBegin_GameThread()
 {
 	check(!PendingCreate.bRenderState);
diff --git a/Engine/Plugins/Experimental/FastGeoStreaming/Source/FastGeoStreaming/Private/FastGeoContainer.h b/Engine/Plugins/Experimental/FastGeoStreaming/Source/FastGeoStreaming/Private/FastGeoContainer.h
--- a/Engine/Plugins/Experimental/FastGeoStreaming/Source/FastGeoStreaming/Private/FastGeoContainer.h
+++ b/Engine/Plugins/Experimental/FastGeoStreaming/Source/FastGeoStreaming/Private/FastGeoContainer.h
@@ -5,6 +5,7 @@
 #include "Engine/AssetUserData.h"
 #include "Physics/Experimental/AsyncPhysicsStateProcessorInterface.h"
 #include "PhysicsEngine/BodyInstance.h"
+#include "AI/Navigation/NavigationElement.h"
 #include "FastGeoComponentCluster.h"
 #include "FastGeoHLOD.h"
 #include "UObject/ObjectPtr.h"
@@ -125,6 +126,8 @@ private:
 	void CollectAssetReferences();
 #endif
 	void SerializeComponentClusters(FArchive& Ar);
+	void RegisterToNavigationSystem();
+	void UnregisterFromNavigationSystem();
 
 	template <typename ThisType, typename TComponentCluster, typename TFunc>
 	static void ForEachComponentCluster(ThisType* Self, TFunc&& InFunc)
@@ -215,6 +218,7 @@ private:
 
 	// Transient data
 	TArray<FFastGeoComponent*> CollisionComponents;
+	TArray<FNavigationElementHandle> NavigationElementHandles;
 
 	bool bIsRegistered = false;
 
diff --git a/Engine/Plugins/Experimental/FastGeoStreaming/Source/FastGeoStreaming/Private/FastGeoInstancedStaticMeshComponent.cpp b/Engine/Plugins/Experimental/FastGeoStreaming/Source/FastGeoStreaming/Private/FastGeoInstancedStaticMeshComponent.cpp
--- a/Engine/Plugins/Experimental/FastGeoStreaming/Source/FastGeoStreaming/Private/FastGeoInstancedStaticMeshComponent.cpp
+++ b/Engine/Plugins/Experimental/FastGeoStreaming/Source/FastGeoStreaming/Private/FastGeoInstancedStaticMeshComponent.cpp
@@ -2,15 +2,18 @@
 
 #include "FastGeoInstancedStaticMeshComponent.h"
 #include "FastGeoComponentCluster.h"
+#include "FastGeoContainer.h"
 #include "FastGeoHLOD.h"
-#include "NaniteSceneProxy.h"
-#include "SceneInterface.h"
-#include "PhysicsEngine/BodySetup.h"
+#include "FastGeoLog.h"
+#include "AI/Navigation/NavCollisionBase.h"
+#include "AI/Navigation/NavigationRelevantData.h"
 #include "Engine/InstancedStaticMesh.h"
 #include "InstancedStaticMeshComponentHelper.h"
 #include "InstancedStaticMesh/ISMInstanceDataSceneProxy.h"
 #include "NaniteVertexFactory.h"
-#include "FastGeoLog.h"
+#include "NaniteSceneProxy.h"
+#include "PhysicsEngine/BodySetup.h"
+#include "SceneInterface.h"
 
 const FFastGeoElementType FFastGeoInstancedStaticMeshComponent::Type(&FFastGeoStaticMeshComponentBase::Type);
 
@@ -60,8 +63,9 @@ void FFastGeoInstancedStaticMeshComponent::InitializeFromComponent(UActorCompone
 	PerInstanceSMCustomData = ISMComponent->PerInstanceSMCustomData;
 	InstancingRandomSeed = ISMComponent->InstancingRandomSeed;
 
-	LocalBounds = CalculateBounds(false);
-	WorldBounds = CalculateBounds(true);
+	LocalBounds = CalculateBounds(EBoundsType::LocalBounds);
+	WorldBounds = CalculateBounds(EBoundsType::WorldBounds);
+	NavigationBounds = CalculateBounds(EBoundsType::NavigationBounds).GetBox();
 }
 
 UClass* FFastGeoInstancedStaticMeshComponent::GetEditorProxyClass() const
@@ -82,7 +86,8 @@ void FFastGeoInstancedStaticMeshComponent::ApplyWorldTransform(const FTransform&
 {
 	Super::ApplyWorldTransform(InTransform);
 
-	WorldBounds = CalculateBounds(true);
+	WorldBounds = CalculateBounds(EBoundsType::WorldBounds);
+	NavigationBounds = CalculateBounds(EBoundsType::NavigationBounds).GetBox();
 }
 
 TSharedPtr<FInstanceDataSceneProxy, ESPMode::ThreadSafe>& FFastGeoInstancedStaticMeshComponent::BuildInstanceData()
@@ -321,11 +326,12 @@ void FFastGeoInstancedStaticMeshComponent::CreateAllInstanceBodies()
 	}
 }
 
-FBoxSphereBounds FFastGeoInstancedStaticMeshComponent::CalculateBounds(bool bWorldSpace)
+FBoxSphereBounds FFastGeoInstancedStaticMeshComponent::CalculateBounds(EBoundsType BoundsType)
 {
 	if (GetStaticMesh() && PerInstanceSMData.Num() > 0)
 	{
-		const FBox InstanceBounds = GetStaticMesh()->GetBounds().GetBox();
+		const bool bWorldSpace = (BoundsType != EBoundsType::LocalBounds);
+		const FBox InstanceBounds = (BoundsType == EBoundsType::NavigationBounds) ? FInstancedStaticMeshComponentHelper::GetInstanceNavigationBounds(*this) : GetStaticMesh()->GetBounds().GetBox();
 		const FMatrix ComponentTransformMatrix = WorldTransform.ToMatrixWithScale();
 		if (InstanceBounds.IsValid)
 		{
@@ -347,6 +353,32 @@ FBoxSphereBounds FFastGeoInstancedStaticMeshComponent::CalculateBounds(bool bWor
 	return FBoxSphereBounds(FVector::ZeroVector, FVector::ZeroVector, 0.f);
 };
 
+void FFastGeoInstancedStaticMeshComponent::GetNavigationData(FNavigationRelevantData& Data) const
+{
+	FInstancedStaticMeshComponentHelper::GetNavigationData(*this, Data, FNavDataPerInstanceTransformDelegate::CreateWeakLambda(GetOwnerContainer(), [this](const FBox& AreaBox, TArray<FTransform>& InstanceData)
+	{
+		FInstancedStaticMeshComponentHelper::GetNavigationPerInstanceTransforms(*this, AreaBox, InstanceData);
+	}));
+}
+
+FBox FFastGeoInstancedStaticMeshComponent::GetNavigationBounds() const
+{
+	return NavigationBounds;
+}
+
+bool FFastGeoInstancedStaticMeshComponent::IsNavigationRelevant() const
+{
+	return PerInstanceSMData.Num() > 0 && Super::IsNavigationRelevant();
+}
+
+bool FFastGeoInstancedStaticMeshComponent::DoCustomNavigableGeometryExport(FNavigableGeometryExport& GeomExport) const
+{
+	return FInstancedStaticMeshComponentHelper::DoCustomNavigableGeometryExport(*this, GeomExport, FNavDataPerInstanceTransformDelegate::CreateWeakLambda(GetOwnerContainer(), [this](const FBox& AreaBox, TArray<FTransform>& InstanceData)
+	{
+		FInstancedStaticMeshComponentHelper::GetNavigationPerInstanceTransforms(*this, AreaBox, InstanceData);
+	}));
+}
+
 #if !WITH_EDITOR && UE_WITH_PSO_PRECACHING
 void FFastGeoInstancedStaticMeshComponent::CollectPSOPrecacheData(const FPSOPrecacheParams& BasePrecachePSOParams, FMaterialInterfacePSOPrecacheParamsList& OutParams)
 {
diff --git a/Engine/Plugins/Experimental/FastGeoStreaming/Source/FastGeoStreaming/Private/FastGeoInstancedStaticMeshComponent.h b/Engine/Plugins/Experimental/FastGeoStreaming/Source/FastGeoStreaming/Private/FastGeoInstancedStaticMeshComponent.h
--- a/Engine/Plugins/Experimental/FastGeoStreaming/Source/FastGeoStreaming/Private/FastGeoInstancedStaticMeshComponent.h
+++ b/Engine/Plugins/Experimental/FastGeoStreaming/Source/FastGeoStreaming/Private/FastGeoInstancedStaticMeshComponent.h
@@ -28,6 +28,14 @@ public:
 #endif
 	//~ End FFastGeoComponent Interface
 
+	//~ Begin FFastGeoPrimitiveComponent interface
+	virtual void GetNavigationData(FNavigationRelevantData& Data) const override;
+	virtual FBox GetNavigationBounds() const override;
+	virtual bool IsNavigationRelevant() const override;
+	virtual bool DoCustomNavigableGeometryExport(FNavigableGeometryExport& GeomExport) const override;
+	virtual bool ShouldSkipNavigationDirtyAreaOnAddOrRemove() override { return true; }
+	//~ End FFastGeoPrimitiveComponent interface
+
 protected:
 	//~ Being FFastGeoComponent interface
 	virtual FLinearColor GetDebugColorImp() const override;
@@ -55,8 +63,14 @@ protected:
 	TSharedPtr<FInstanceDataSceneProxy, ESPMode::ThreadSafe>& BuildInstanceData();
 
 private:
+	enum class EBoundsType
+	{
+		LocalBounds,
+		WorldBounds,
+		NavigationBounds
+	};
+	FBoxSphereBounds CalculateBounds(EBoundsType BoundsType);
 	void CreateAllInstanceBodies();
-	FBoxSphereBounds CalculateBounds(bool bWorldSpace);
 
 private:
 
@@ -74,12 +88,15 @@ private:
 	// Transient data
 	TSharedPtr<FInstanceDataSceneProxy, ESPMode::ThreadSafe> DataProxy{};
 	TArray<float> InstanceRandomIDs;
+	FBox NavigationBounds;
 	
 	/** Physics representation of the instance bodies. */
 	TArray<FBodyInstance*> InstanceBodies;
 
 	/** Payload used by asynchronous destruction of physics state (see OnAsyncDestroyPhysicsState). */
 	TArray<FBodyInstance*> AsyncDestroyPhysicsStatePayload;
+
+	friend class FInstancedStaticMeshComponentHelper;
 };
 
 UCLASS()
diff --git a/Engine/Plugins/Experimental/FastGeoStreaming/Source/FastGeoStreaming/Private/FastGeoPrimitiveComponent.cpp b/Engine/Plugins/Experimental/FastGeoStreaming/Source/FastGeoStreaming/Private/FastGeoPrimitiveComponent.cpp
--- a/Engine/Plugins/Experimental/FastGeoStreaming/Source/FastGeoStreaming/Private/FastGeoPrimitiveComponent.cpp
+++ b/Engine/Plugins/Experimental/FastGeoStreaming/Source/FastGeoStreaming/Private/FastGeoPrimitiveComponent.cpp
@@ -7,6 +7,9 @@
 #include "FastGeoWorldSubsystem.h"
 #include "FastGeoComponentCluster.h"
 #include "FastGeoWeakElement.h"
+#include "FastGeoLog.h"
+#include "AI/NavigationModifier.h"
+#include "AI/Navigation/NavigationRelevantData.h"
 #include "PrimitiveSceneDesc.h"
 #include "SceneInterface.h"
 #include "WorldPartition/WorldPartition.h"
@@ -14,7 +17,7 @@
 #include "PhysicsEngine/BodySetup.h"
 #include "PSOPrecache.h"
 #include "PSOPrecacheMaterial.h"
-#include "FastGeoLog.h"
+#include "PrimitiveComponentHelper.h"
 
 #if !UE_BUILD_SHIPPING
 namespace FastGeo
@@ -68,8 +71,12 @@ void FFastGeoPrimitiveComponent::Serialize(FArchive& Ar)
 	Ar << WorldBounds;
 	FArchive_Serialize_BitfieldBool(Ar, bIsVisible);
 	FArchive_Serialize_BitfieldBool(Ar, bStaticWhenNotMoveable);
+	FArchive_Serialize_BitfieldBool(Ar, bFillCollisionUnderneathForNavmesh);
+	FArchive_Serialize_BitfieldBool(Ar, bRasterizeAsFilledConvexVolume);
+	FArchive_Serialize_BitfieldBool(Ar, bCanEverAffectNavigation);
 	Ar << CustomPrimitiveData.Data;
 	Ar << DetailMode;
+	Ar << bHasCustomNavigableGeometry;
 	Ar << RuntimeVirtualTextures;
 	FBodyInstance::StaticStruct()->SerializeItem(Ar, &BodyInstance, nullptr);
 
@@ -164,8 +171,12 @@ void FFastGeoPrimitiveComponent::InitializeFromComponent(UActorComponent* Compon
 	WorldTransform = LocalTransform;
 	bIsVisible = PrimitiveComponent->IsVisible();
 	bStaticWhenNotMoveable = PrimitiveComponent->GetStaticWhenNotMoveable();
+	bFillCollisionUnderneathForNavmesh = PrimitiveComponent->bFillCollisionUnderneathForNavmesh;
+	bRasterizeAsFilledConvexVolume = PrimitiveComponent->bRasterizeAsFilledConvexVolume;
+	bCanEverAffectNavigation = PrimitiveComponent->CanEverAffectNavigation();
 	CustomPrimitiveData = PrimitiveComponent->GetCustomPrimitiveData();
 	DetailMode = PrimitiveComponent->DetailMode;
+	bHasCustomNavigableGeometry = PrimitiveComponent->bHasCustomNavigableGeometry;
 	BodyInstance.CopyBodyInstancePropertiesFrom(&PrimitiveComponent->BodyInstance);
 	RuntimeVirtualTextures = PrimitiveComponent->GetRuntimeVirtualTextures();
 
@@ -488,12 +499,12 @@ void FFastGeoPrimitiveComponent::MarkRenderStateDirty()
 	}
 }
 
-FTransform FFastGeoPrimitiveComponent::GetTransform() const
+const FTransform& FFastGeoPrimitiveComponent::GetTransform() const
 {
 	return WorldTransform;
 }
 
-FBoxSphereBounds FFastGeoPrimitiveComponent::GetBounds() const
+const FBoxSphereBounds& FFastGeoPrimitiveComponent::GetBounds() const
 {
 	return WorldBounds;
 }
@@ -669,6 +680,57 @@ void FFastGeoPrimitiveComponent::PrecachePSOs()
 #endif
 }
 
+bool FFastGeoPrimitiveComponent::IsNavigationRelevant() const
+{
+	if (!bCanEverAffectNavigation)
+	{
+		return false;
+	}
+
+	if (HasCustomNavigableGeometry() >= EHasCustomNavigableGeometry::EvenIfNotCollidable)
+	{
+		return true;
+	}
+
+	auto GetCollisionEnabled = [this]() -> ECollisionEnabled::Type
+	{
+		if (!IsCollisionEnabled())
+		{
+			return ECollisionEnabled::NoCollision;
+		}
+		return BodyInstance.GetCollisionEnabled(false);
+	};
+
+	auto IsQueryCollisionEnabled = [&GetCollisionEnabled]()
+	{
+		return CollisionEnabledHasQuery(GetCollisionEnabled());
+	};
+
+	auto GetCollisionResponseToChannels = [this]()
+	{
+		return BodyInstance.GetResponseToChannels();
+	};
+
+	const FCollisionResponseContainer& ResponseToChannels = GetCollisionResponseToChannels();
+	return IsQueryCollisionEnabled() &&
+		(ResponseToChannels.GetResponse(ECC_Pawn) == ECR_Block || ResponseToChannels.GetResponse(ECC_Vehicle) == ECR_Block);
+}
+
+FBox FFastGeoPrimitiveComponent::GetNavigationBounds() const
+{
+	return GetBounds().GetBox();
+}
+
+void FFastGeoPrimitiveComponent::GetNavigationData(FNavigationRelevantData& OutData) const
+{
+	FPrimitiveComponentHelper::GetNavigationData(*this, OutData);
+}
+
+EHasCustomNavigableGeometry::Type FFastGeoPrimitiveComponent::HasCustomNavigableGeometry() const
+{
+	return bHasCustomNavigableGeometry;
+}
+
 //~ Begin IPhysicsBodyInstanceOwner interface
 bool FFastGeoPrimitiveComponent::IsStaticPhysics() const
 {
diff --git a/Engine/Plugins/Experimental/FastGeoStreaming/Source/FastGeoStreaming/Private/FastGeoPrimitiveComponent.h b/Engine/Plugins/Experimental/FastGeoStreaming/Source/FastGeoStreaming/Private/FastGeoPrimitiveComponent.h
--- a/Engine/Plugins/Experimental/FastGeoStreaming/Source/FastGeoStreaming/Private/FastGeoPrimitiveComponent.h
+++ b/Engine/Plugins/Experimental/FastGeoStreaming/Source/FastGeoStreaming/Private/FastGeoPrimitiveComponent.h
@@ -50,11 +50,14 @@ public:
 		
 	FSceneInterface* GetScene() const;
 	FPrimitiveSceneProxy* GetSceneProxy() const;
-	FTransform GetTransform() const;
-	FBoxSphereBounds GetBounds() const;
+	const FTransform& GetTransform() const;
+	const FBoxSphereBounds& GetBounds() const;
 	FMatrix GetRenderMatrix() const;
 	float GetLastRenderTimeOnScreen() const;
 
+	// Used by FStaticMeshComponentHelper/FInstancedStaticMeshComponentHelper
+	const FTransform& GetComponentTransform() const { return GetTransform(); }
+
 	void CreateRenderState(FRegisterComponentContext* Context);
 
 	struct FFastGeoDestroyRenderStateContext
@@ -99,6 +102,16 @@ public:
 	virtual void InitializeFromComponent(UActorComponent* Component) override;
 	virtual UClass* GetEditorProxyClass() const override;
 #endif
+	//~End FFastGeoComponent interface
+
+	//~ Begin Navigation
+	virtual bool IsNavigationRelevant() const;
+	virtual bool ShouldSkipNavigationDirtyAreaOnAddOrRemove() { return false; }
+	virtual FBox GetNavigationBounds() const;
+	virtual void GetNavigationData(FNavigationRelevantData& OutData) const;
+	virtual EHasCustomNavigableGeometry::Type HasCustomNavigableGeometry() const;
+	virtual bool DoCustomNavigableGeometryExport(FNavigableGeometryExport& GeomExport) const { return true; }
+	//~ End Navigation
 
 protected:
 	virtual FPrimitiveSceneProxyDesc& GetSceneProxyDesc() = 0;
@@ -121,8 +134,12 @@ protected:
 	FBoxSphereBounds WorldBounds;
 	uint8 bIsVisible : 1 = true;
 	uint8 bStaticWhenNotMoveable : 1 = true;
+	uint8 bFillCollisionUnderneathForNavmesh : 1 = false;
+	uint8 bRasterizeAsFilledConvexVolume : 1 = false;
+	uint8 bCanEverAffectNavigation : 1 = false;
 	FCustomPrimitiveData CustomPrimitiveData;
 	TEnumAsByte<enum EDetailMode> DetailMode;
+	TEnumAsByte<EHasCustomNavigableGeometry::Type> bHasCustomNavigableGeometry;
 	FBodyInstance BodyInstance;
 	TArray<URuntimeVirtualTexture*> RuntimeVirtualTextures;
 
@@ -168,6 +185,7 @@ protected:
 	static_assert((int)EPSOPrecachePriority::Highest < 1 << 2);
 #endif
 
+	friend class FPrimitiveComponentHelper;
 	friend class UFastGeoContainer;
 	friend class UFastGeoPrimitiveComponentEditorProxy;
 };
diff --git a/Engine/Plugins/Experimental/FastGeoStreaming/Source/FastGeoStreaming/Private/FastGeoStaticMeshComponent.cpp b/Engine/Plugins/Experimental/FastGeoStreaming/Source/FastGeoStreaming/Private/FastGeoStaticMeshComponent.cpp
--- a/Engine/Plugins/Experimental/FastGeoStreaming/Source/FastGeoStreaming/Private/FastGeoStaticMeshComponent.cpp
+++ b/Engine/Plugins/Experimental/FastGeoStreaming/Source/FastGeoStreaming/Private/FastGeoStaticMeshComponent.cpp
@@ -2,6 +2,9 @@
 
 #include "FastGeoStaticMeshComponent.h"
 #include "FastGeoContainer.h"
+#include "FastGeoLog.h"
+#include "AI/Navigation/NavCollisionBase.h"
+#include "AI/Navigation/NavigationRelevantData.h"
 #include "Materials/MaterialRelevance.h"
 #include "SceneInterface.h"
 #include "MeshComponentHelper.h"
@@ -11,7 +14,6 @@
 #include "PhysicalMaterials/PhysicalMaterial.h"
 #include "NaniteVertexFactory.h"
 #include "StaticMeshSceneProxy.h"
-#include "FastGeoLog.h"
 
 #if WITH_EDITOR
 #include "ObjectCacheEventSink.h"
@@ -75,6 +77,31 @@ UBodySetup* FFastGeoStaticMeshComponentBase::GetBodySetup() const
 	return GetStaticMesh() ? GetStaticMesh()->GetBodySetup() : nullptr;
 }
 
+bool FFastGeoStaticMeshComponentBase::IsNavigationRelevant() const
+{
+	return FStaticMeshComponentHelper::IsNavigationRelevant(*this);
+}
+
+FBox FFastGeoStaticMeshComponentBase::GetNavigationBounds() const
+{
+	return FStaticMeshComponentHelper::GetNavigationBounds(*this);
+}
+
+void FFastGeoStaticMeshComponentBase::GetNavigationData(FNavigationRelevantData& Data) const
+{
+	FStaticMeshComponentHelper::GetNavigationData(*this, Data);
+}
+
+bool FFastGeoStaticMeshComponentBase::DoCustomNavigableGeometryExport(FNavigableGeometryExport& GeomExport) const
+{
+	return FStaticMeshComponentHelper::DoCustomNavigableGeometryExport(*this, GeomExport);
+}
+
+bool FFastGeoStaticMeshComponentBase::ShouldExportAsObstacle(const UNavCollisionBase& InNavCollision) const
+{
+	return InNavCollision.IsDynamicObstacle();
+}
+
 int32 FFastGeoStaticMeshComponentBase::GetNumMaterials() const
 {
 	return GetStaticMesh() ? GetStaticMesh()->GetStaticMaterials().Num() : 0;
diff --git a/Engine/Plugins/Experimental/FastGeoStreaming/Source/FastGeoStreaming/Private/FastGeoStaticMeshComponent.h b/Engine/Plugins/Experimental/FastGeoStreaming/Source/FastGeoStreaming/Private/FastGeoStaticMeshComponent.h
--- a/Engine/Plugins/Experimental/FastGeoStreaming/Source/FastGeoStreaming/Private/FastGeoStaticMeshComponent.h
+++ b/Engine/Plugins/Experimental/FastGeoStreaming/Source/FastGeoStreaming/Private/FastGeoStaticMeshComponent.h
@@ -8,6 +8,7 @@
 
 class UBodySetup;
 class UMaterialInterface;
+class UNavCollisionBase;
 class FVertexFactoryType;
 struct FMaterialRelevance;
 struct FStaticMeshLODResources;
@@ -60,6 +61,13 @@ public:
 #endif
 	//~ End FFastGeoComponent interface
 
+	//~ Begin FFastGeoPrimitiveComponent interface
+	virtual bool IsNavigationRelevant() const override;
+	virtual FBox GetNavigationBounds() const override;
+	virtual void GetNavigationData(FNavigationRelevantData& Data) const override;
+	virtual bool DoCustomNavigableGeometryExport(FNavigableGeometryExport& GeomExport) const override;
+	//~ End FFastGeoPrimitiveComponent interface
+
 	//~ Begin IPhysicsBodyInstanceOwner interface
 	virtual UPhysicalMaterial* GetPhysicalMaterial() const override;
 	virtual void GetComplexPhysicalMaterials(TArray<UPhysicalMaterial*>& OutPhysMaterials, TArray<FPhysicalMaterialMaskParams>* OutPhysMaterialMasks) const override;
@@ -86,6 +94,7 @@ protected:
 	UMaterialInterface* GetMaterial(int32 MaterialIndex, bool bDoingNaniteMaterialAudit) const;
 	void GetDefaultMaterialSlotsOverlayMaterial(TArray<TObjectPtr<UMaterialInterface>>& OutMaterialSlotOverlayMaterials) const;
 	FPrimitiveMaterialPropertyDescriptor GetUsedMaterialPropertyDesc(ERHIFeatureLevel::Type FeatureLevel) const;
+	bool ShouldExportAsObstacle(const UNavCollisionBase& InNavCollision) const;
 #if !WITH_EDITOR && UE_WITH_PSO_PRECACHING
 	virtual void CollectPSOPrecacheData(const FPSOPrecacheParams& BasePrecachePSOParams, FMaterialInterfacePSOPrecacheParamsList& OutParams);
 #endif
