FastGeo : UFastGeoWorldPartitionRuntimeCellTransformer now tests if Water module is loaded before adding dissallowed Water classes.

diff --git a/Engine/Plugins/Experimental/FastGeoStreaming/Source/FastGeoStreaming/Private/FastGeoWorldPartitionRuntimeCellTransformer.cpp b/Engine/Plugins/Experimental/FastGeoStreaming/Source/FastGeoStreaming/Private/FastGeoWorldPartitionRuntimeCellTransformer.cpp
--- a/Engine/Plugins/Experimental/FastGeoStreaming/Source/FastGeoStreaming/Private/FastGeoWorldPartitionRuntimeCellTransformer.cpp
+++ b/Engine/Plugins/Experimental/FastGeoStreaming/Source/FastGeoStreaming/Private/FastGeoWorldPartitionRuntimeCellTransformer.cpp
@@ -13,6 +13,7 @@
 #include "Components/StaticMeshComponent.h"
 #include "Components/InstancedStaticMeshComponent.h"
 #include "Components/SplineMeshComponent.h"
+#include "Modules/ModuleManager.h"
 #include "FastGeoContainer.h"
 #include "FastGeoComponentCluster.h"
 #include "FastGeoComponent.h"
@@ -61,11 +62,15 @@ UFastGeoWorldPartitionRuntimeCellTransformer::UFastGeoWorldPartitionRuntimeCellT
 	AllowedComponentClasses.Add(UStaticMeshComponent::StaticClass());
 
 	// Disallowed actor/component classes
-	FastGeo::FClassHelper::AddClass(TEXT("/Script/Water.WaterBody"), DisallowedActorClasses);
 	FastGeo::FClassHelper::AddClass(TEXT("/Script/Foliage.InteractiveFoliageComponent"), DisallowedComponentClasses);
 	FastGeo::FClassHelper::AddClass(TEXT("/Script/Landscape.LandscapeNaniteComponent"), DisallowedComponentClasses);
 	FastGeo::FClassHelper::AddClass(TEXT("/Script/Landscape.LandscapeMeshProxyComponent"), DisallowedComponentClasses);
-	FastGeo::FClassHelper::AddClass(TEXT("/Script/Water.WaterBodyMeshComponent"), DisallowedComponentClasses);
+	if (FModuleManager::Get().LoadModule(TEXT("Water")))
+	{
+		FastGeo::FClassHelper::AddClass(TEXT("/Script/Water.WaterBodyMeshComponent"), DisallowedComponentClasses);
+		FastGeo::FClassHelper::AddClass(TEXT("/Script/Water.WaterBody"), DisallowedActorClasses);
+	}
+
 	DisallowedComponentClasses.Add(USplineMeshComponent::StaticClass());
 
 	// Ignored component classes
