FastGeo: Fixed ISM NavigationBounds not persisted for cooked builds.

diff --git a/Engine/Plugins/Experimental/FastGeoStreaming/Source/FastGeoStreaming/Private/FastGeoInstancedStaticMeshComponent.cpp b/Engine/Plugins/Experimental/FastGeoStreaming/Source/FastGeoStreaming/Private/FastGeoInstancedStaticMeshComponent.cpp
--- a/Engine/Plugins/Experimental/FastGeoStreaming/Source/FastGeoStreaming/Private/FastGeoInstancedStaticMeshComponent.cpp
+++ b/Engine/Plugins/Experimental/FastGeoStreaming/Source/FastGeoStreaming/Private/FastGeoInstancedStaticMeshComponent.cpp
@@ -31,6 +31,7 @@ void FFastGeoInstancedStaticMeshComponent::Serialize(FArchive& Ar)
 	Ar << InstancingRandomSeed;
 	PerInstanceSMCustomData.BulkSerialize(Ar);
 	Ar << AdditionalRandomSeeds;
+	Ar << NavigationBounds;
 
 	// Serialize persistent data from FInstancedStaticMeshSceneProxyDesc
 	Ar << SceneProxyDesc.InstanceLODDistanceScale;
