FastGeo: Added debug mode and debug mode on editor selection in FastGeoCellTransformer to allow to debug/understand why actors/components are not transformed.
Also outputs stats on resulting transformation.

diff --git a/Engine/Plugins/Experimental/FastGeoStreaming/Source/FastGeoStreaming/FastGeoStreaming.Build.cs b/Engine/Plugins/Experimental/FastGeoStreaming/Source/FastGeoStreaming/FastGeoStreaming.Build.cs
--- a/Engine/Plugins/Experimental/FastGeoStreaming/Source/FastGeoStreaming/FastGeoStreaming.Build.cs
+++ b/Engine/Plugins/Experimental/FastGeoStreaming/Source/FastGeoStreaming/FastGeoStreaming.Build.cs
@@ -18,6 +18,11 @@ namespace UnrealBuildTool.Rules
 					"RHI"
 				}
 			);
+
+			if (Target.bBuildEditor == true)
+			{
+				PrivateDependencyModuleNames.Add("UnrealEd");
+			}
 		}
 	}
 }
\ No newline at end of file
diff --git a/Engine/Plugins/Experimental/FastGeoStreaming/Source/FastGeoStreaming/Private/FastGeoWorldPartitionRuntimeCellTransformer.cpp b/Engine/Plugins/Experimental/FastGeoStreaming/Source/FastGeoStreaming/Private/FastGeoWorldPartitionRuntimeCellTransformer.cpp
--- a/Engine/Plugins/Experimental/FastGeoStreaming/Source/FastGeoStreaming/Private/FastGeoWorldPartitionRuntimeCellTransformer.cpp
+++ b/Engine/Plugins/Experimental/FastGeoStreaming/Source/FastGeoStreaming/Private/FastGeoWorldPartitionRuntimeCellTransformer.cpp
@@ -10,18 +10,23 @@
 #include "Engine/BlueprintGeneratedClass.h"
 #include "EdGraph/EdGraph.h"
 #include "UObject/Package.h"
-#include "Components/StaticMeshComponent.h"
+#include "Modules/ModuleManager.h"
+#include "Physics/Experimental/PhysScene_Chaos.h"
+#include "Streaming/ActorTextureStreamingBuildDataComponent.h"
 #include "Components/InstancedStaticMeshComponent.h"
+#include "Components/StaticMeshComponent.h"
 #include "Components/SplineMeshComponent.h"
-#include "Modules/ModuleManager.h"
+#include "GameFramework/WorldSettings.h"
+#include "GameFramework/Info.h"
+#include "WorldPartition/WorldPartition.h"
+#include "WorldPartition/WorldPartitionHelpers.h"
+#include "Selection.h"
 #include "FastGeoContainer.h"
 #include "FastGeoComponentCluster.h"
 #include "FastGeoComponent.h"
 #include "FastGeoHLOD.h"
 #include "FastGeoStaticMeshComponent.h"
 #include "FastGeoInstancedStaticMeshComponent.h"
-#include "Physics/Experimental/PhysScene_Chaos.h"
-#include "Streaming/ActorTextureStreamingBuildDataComponent.h"
 #include "FastGeoLog.h"
 #endif
 
@@ -34,9 +39,32 @@
 
 #define LOCTEXT_NAMESPACE "WorldPartition"
 
-#if WITH_EDITORONLY_DATA
+#if WITH_EDITOR
 namespace FastGeo
 {
+	// Used to cancel package being dirtied when bDebugMode is modified (see PostEditChangeProperty)
+	static bool GPackageWasDirty = false;
+
+	// Tag use to force include actors into FastGeoStreaming
+	static const FName NAME_FastGeo(TEXT("FastGeo"));
+
+	// Tag use to force exclude actors from FastGeoStreaming
+	static const FName NAME_NoFastGeo(TEXT("NoFastGeo"));
+
+	static bool IsCollisionEnabled(UStaticMeshComponent* Component)
+	{
+		return Component->IsCollisionEnabled() && !Component->GetOwner()->IsA<AWorldPartitionHLOD>();
+	}
+
+	static const TCHAR* GetComponentShortName(UActorComponent* InComponent)
+	{
+		TStringBuilder<256> Builder;
+		Builder += InComponent->GetName();
+		Builder += TEXT(".");
+		Builder += InComponent->GetOwner()->GetName();
+		return *Builder;
+	}
+
 	struct FClassHelper
 	{
 		template <typename T>
@@ -48,12 +76,44 @@ namespace FastGeo
 			}
 		}
 	};
+
+	struct FTransformationStats
+	{
+		int32 TotalActorCount = 0;
+		int32 TotalComponentCount = 0;
+		int32 FullyTransformableActorCount = 0;
+		int32 PartiallyTransformableActorCount = 0;
+		int32 TransformedComponentCount = 0;
+
+		void DumpStats(const TCHAR* InPrefixString)
+		{
+			if (TotalActorCount)
+			{
+				const float FullyTransformedActorPercentage = TotalActorCount > 0 ? (100.f * FullyTransformableActorCount) / TotalActorCount : 0.f;
+				const float PartiallyTransformedActorPercentage = TotalActorCount > 0 ? (100.f * PartiallyTransformableActorCount) / TotalActorCount : 0.f;
+				const float TransformedComponentPercentage = TotalComponentCount > 0 ? (100.f * TransformedComponentCount) / TotalComponentCount : 0.f;
+				const int32 NonTransformableActorCount = FMath::Max(0, TotalActorCount - FullyTransformableActorCount - PartiallyTransformableActorCount);
+				const float NonTransformableActorPercentage = NonTransformableActorCount > 0 ? (100.f * NonTransformableActorCount) / TotalActorCount : 0.f;
+
+				UE_CLOG(FullyTransformableActorCount, LogFastGeoStreaming, Log, TEXT("%s Transformable Actors (Full)    = %d (%3.1f%%)"), InPrefixString, FullyTransformableActorCount, FullyTransformedActorPercentage);
+				UE_CLOG(PartiallyTransformableActorCount,LogFastGeoStreaming, Log, TEXT("%s Transformable Actors (Partial) = %d (%3.1f%%)"), InPrefixString, PartiallyTransformableActorCount, PartiallyTransformedActorPercentage);
+				UE_CLOG(TransformedComponentCount, LogFastGeoStreaming, Log, TEXT("%s Transformable Components       = %d (%3.1f%%)"), InPrefixString, TransformedComponentCount, TransformedComponentPercentage);
+				UE_CLOG(NonTransformableActorCount, LogFastGeoStreaming, Log, TEXT("%s Non-Transformable Actors       = %d (%3.1f%%)"), InPrefixString, NonTransformableActorCount, NonTransformableActorPercentage);
+			}
+		}
+	};
 }
 #endif
 
 UFastGeoWorldPartitionRuntimeCellTransformer::UFastGeoWorldPartitionRuntimeCellTransformer(const FObjectInitializer& ObjectInitializer)
 	: Super(ObjectInitializer)
 {
+#if WITH_EDITOR
+	if (GIsEditor && !IsTemplate())
+	{
+		USelection::SelectionChangedEvent.AddUObject(this, &UFastGeoWorldPartitionRuntimeCellTransformer::OnSelectionChanged);	}
+#endif
+
 #if WITH_EDITORONLY_DATA
 	// Allowed actor/component classes
 	AllowedActorClasses.Add(APackedLevelActor::StaticClass());
@@ -70,7 +130,7 @@ UFastGeoWorldPartitionRuntimeCellTransformer::UFastGeoWorldPartitionRuntimeCellT
 		FastGeo::FClassHelper::AddClass(TEXT("/Script/Water.WaterBodyMeshComponent"), DisallowedComponentClasses);
 		FastGeo::FClassHelper::AddClass(TEXT("/Script/Water.WaterBody"), DisallowedActorClasses);
 	}
-
+	DisallowedActorClasses.Add(AInfo::StaticClass());
 	DisallowedComponentClasses.Add(USplineMeshComponent::StaticClass());
 
 	// Ignored component classes
@@ -79,38 +139,102 @@ UFastGeoWorldPartitionRuntimeCellTransformer::UFastGeoWorldPartitionRuntimeCellT
 #endif
 }
 
-static const FName NAME_FastGeo(TEXT("FastGeo"));
-static const FName NAME_NoFastGeo(TEXT("NoFastGeo"));
+void UFastGeoWorldPartitionRuntimeCellTransformer::Serialize(FArchive& Ar)
+{
+	Super::Serialize(Ar);
 
 #if WITH_EDITOR
+	if (Ar.GetPortFlags() & PPF_DuplicateForPIE)
+	{
+		Ar << bDebugMode;
+	}
+#endif
+}
 
-namespace FastGeo
+#if WITH_EDITOR
+
+bool UFastGeoWorldPartitionRuntimeCellTransformer::IsDebugModeEnabled = false;
+FAutoConsoleVariableRef UFastGeoWorldPartitionRuntimeCellTransformer::CVarIsDebugModeEnabled(
+	TEXT("FastGeo.EnableTransformerDebugMode"),
+	UFastGeoWorldPartitionRuntimeCellTransformer::IsDebugModeEnabled,
+	TEXT("Set to true to enable FastGeoStreaming transformer debug mode (used in PIE and at cook time)."),
+	ECVF_Default);
+
+bool UFastGeoWorldPartitionRuntimeCellTransformer::IsFastGeoEnabled = true;
+FAutoConsoleVariableRef UFastGeoWorldPartitionRuntimeCellTransformer::CVarIsFastGeoEnabled(
+	TEXT("FastGeo.Enable"),
+	UFastGeoWorldPartitionRuntimeCellTransformer::IsFastGeoEnabled,
+	TEXT("Set to false to disable FastGeoStreaming (used when starting PIE and at cook time)."),
+	ECVF_Default);
+
+void UFastGeoWorldPartitionRuntimeCellTransformer::BeginDestroy()
 {
-	auto IsCollisionEnabled = [](UStaticMeshComponent* Component)
+	Super::BeginDestroy();
+
+	if (GIsEditor && !IsTemplate())
 	{
-		return Component->IsCollisionEnabled() && !Component->GetOwner()->IsA<AWorldPartitionHLOD>();
-	};
+		USelection::SelectionChangedEvent.RemoveAll(this);
+	}
+}
+
+bool UFastGeoWorldPartitionRuntimeCellTransformer::IsDebugMode() const
+{
+	return bDebugMode || UFastGeoWorldPartitionRuntimeCellTransformer::IsDebugModeEnabled;
+}
 
-	class FNullArchive : public FArchive
+void UFastGeoWorldPartitionRuntimeCellTransformer::OnSelectionChanged(UObject* Object)
+{
+	if (bDebugModeOnEditorSelection && GetOuterUWorldPartition()->GetWorld() && !GetOuterUWorldPartition()->GetWorld()->IsGameWorld())
 	{
-		//override FObjectPtr to do nothing
-		virtual FArchive& operator<<(struct FObjectPtr& Value) override
+		TArray<AActor*> SelectedActors;
+		if (USelection* Selection = Cast<USelection>(Object))
 		{
-			return *this;
+			for (int32 Index = 0; Index < Selection->Num(); Index++)
+			{
+				if (AActor* SelectedActor = Cast<AActor>(Selection->GetSelectedObject(Index)))
+				{
+					UWorldPartition* WorldPartition = FWorldPartitionHelpers::GetWorldPartition(SelectedActor);
+					if (WorldPartition == GetOuterUWorldPartition())
+					{
+						SelectedActors.Add(SelectedActor);
+					}
+				}
+			}
 		}
-	};
 
-	static bool bFastGeoEnabled = true;
-	static FAutoConsoleVariableRef CVarFastGeoEnabled(
-		TEXT("FastGeo.Enable"),
-		bFastGeoEnabled,
-		TEXT("Set to false to disable FastGeoStreaming (used when starting PIE and at cook time)."),
-		ECVF_Default);
+		if (!SelectedActors.IsEmpty())
+		{
+			TGuardValue<bool> Guard(FFastGeoTransformResult::ShouldReport, true);
+
+			UE_LOG(LogFastGeoStreaming, Log, TEXT("------------------------------------------------------------------------"));
+			UE_LOG(LogFastGeoStreaming, Log, TEXT("- FastGeoStreaming Debug Mode: Transformation on %d selected actors "), SelectedActors.Num());
+
+			FastGeo::FTransformationStats Stats;
+			Stats.TotalActorCount = SelectedActors.Num();
+			for (AActor* Actor : SelectedActors)
+			{
+				bool bIsActorFullyTransformable;
+				TSet<UActorComponent*> TransformableComponents;
+				if (CanTransformActor(Actor, bIsActorFullyTransformable, TransformableComponents))
+				{
+					Stats.FullyTransformableActorCount += bIsActorFullyTransformable ? 1 : 0;
+					Stats.PartiallyTransformableActorCount += !bIsActorFullyTransformable ? 1 : 0;
+					Stats.TransformedComponentCount += TransformableComponents.Num();
+				}
+				Stats.TotalComponentCount += Actor->GetComponents().Num();
+			}
+
+			Stats.DumpStats(TEXT("  - "));
+			UE_LOG(LogFastGeoStreaming, Log, TEXT("------------------------------------------------------------------------"));
+		}
+	}
 }
 
 void UFastGeoWorldPartitionRuntimeCellTransformer::Transform(ULevel* InLevel)
 {
-	if (!FastGeo::bFastGeoEnabled)
+	TGuardValue<bool> Guard(FFastGeoTransformResult::ShouldReport, IsDebugMode());
+
+	if (!UFastGeoWorldPartitionRuntimeCellTransformer::IsFastGeoEnabled)
 	{
 		return;
 	}
@@ -123,76 +247,34 @@ void UFastGeoWorldPartitionRuntimeCellTransformer::Transform(ULevel* InLevel)
 	}
 
 	check(InLevel);
-
-	auto IsAllowedComponentClass = [this](UActorComponent* Component)
+	
+	if (IsDebugMode())
 	{
-		UClass* ComponentClass = Component->GetClass();
-		for (TSubclassOf<UActorComponent> DisallowedComponentClass : DisallowedComponentClasses)
-		{
-			if (ComponentClass->IsChildOf(DisallowedComponentClass))
-			{
-				return false;
-			}
-		}
+		UE_LOG(LogFastGeoStreaming, Log, TEXT("------------------------------------------------------------------------"));
+		UE_LOG(LogFastGeoStreaming, Log, TEXT("- FastGeoStreaming Debug Mode: Transforming Level '%s'"), *InLevel->GetPathName());
+	}
 
-		for (TSubclassOf<UActorComponent> DisallowedComponentClass : DisallowedExactComponentClasses)
+	FastGeo::FTransformationStats Stats;
+	TMap<TObjectPtr<AActor>*, TPair<bool, TSet<UActorComponent*>>> ActorComponents;
+	for (TObjectPtr<AActor>& Actor : InLevel->Actors)
+	{
+		// Always skip the AWorldSettings of a Level (don't even count it)
+		if (IsValid(Actor) && !Actor->IsA<AWorldSettings>())
 		{
-			if (ComponentClass == DisallowedComponentClass)
-			{
-				return false;
-			}
-		}
+			Stats.TotalActorCount++;
+			Stats.TotalComponentCount += Actor->GetComponents().Num();
 
-		for (TSubclassOf<UActorComponent> AllowedComponentClass : AllowedComponentClasses)
-		{
-			if (ComponentClass->IsChildOf(AllowedComponentClass))
+			bool bIsActorFullyTransformable;
+			TSet<UActorComponent*> TransformableComponents;
+			if (CanTransformActor(Actor, bIsActorFullyTransformable, TransformableComponents))
 			{
-				return true;
-			}
-		}
+				Stats.FullyTransformableActorCount += bIsActorFullyTransformable ? 1 : 0;
+				Stats.PartiallyTransformableActorCount += !bIsActorFullyTransformable ? 1 : 0;
+				Stats.TransformedComponentCount += TransformableComponents.Num();
 
-		for (TSubclassOf<UActorComponent> AllowedComponentClass : AllowedExactComponentClasses)
-		{
-			if (ComponentClass == AllowedComponentClass)
-			{
-				return true;
+				ActorComponents.FindOrAdd(&Actor, { bIsActorFullyTransformable, MoveTemp(TransformableComponents) });
 			}
 		}
-
-		return false;
-	};
-
-	TMap<TObjectPtr<AActor>*, TSet<UActorComponent*>> ActorComponents;
-	for (TObjectPtr<AActor>& Actor : InLevel->Actors)
-	{
-		if (IsValid(Actor) && CanTransformActor(Actor))
-		{
-			// Gather components
-			Actor->ForEachComponent<UStaticMeshComponent>(false, [&Actor, &ActorComponents, &IsAllowedComponentClass](UStaticMeshComponent* StaticMeshComponent)
-			{
-				if (!StaticMeshComponent->IsEditorOnly() &&
-					StaticMeshComponent->GetStaticMesh() &&
-					!StaticMeshComponent->GetLODParentPrimitive() &&
-					(StaticMeshComponent->Mobility == EComponentMobility::Static) &&
-					IsAllowedComponentClass(StaticMeshComponent))
-				{
-					// Make sure BodyInstance CollisionEnabled is updated first
-					StaticMeshComponent->UpdateCollisionFromStaticMesh();
-					const bool bIsCollisionEnabled = FastGeo::IsCollisionEnabled(StaticMeshComponent);
-
-					// If collision is enabled, only allow if async physics state creation and destruction are supported 
-					// If colliion is disabled, allow if component is visible
-					if (!bIsCollisionEnabled ||
-						(StaticMeshComponent->AllowsAsyncPhysicsStateCreation() && StaticMeshComponent->AllowsAsyncPhysicsStateDestruction()))
-					{
-						if (bIsCollisionEnabled || StaticMeshComponent->IsVisible())
-						{
-							ActorComponents.FindOrAdd(&Actor).Add(StaticMeshComponent);
-						}
-					}
-				}
-			});
-		}
 	}
 
 	const TMap<TSubclassOf<UPrimitiveComponent>, FFastGeoElementType> FastGeoComponentTypeMapping =
@@ -224,19 +306,8 @@ void UFastGeoWorldPartitionRuntimeCellTransformer::Transform(ULevel* InLevel)
 
 		TUniquePtr<FFastGeoComponentCluster> LevelComponentCluster(new FFastGeoComponentCluster(FastGeo, *FString::Printf(TEXT("FastGeoComponentCluster_%s"), *CellName)));
 
-		for (auto& [Actor, Components] : ActorComponents)
+		for (auto& [Actor, ActorInfo] : ActorComponents)
 		{
-			check(!Components.IsEmpty());
-
-			// Can't convert partially a BP actors
-			// Rerun CS will be called in PIE when registering the component and
-			// also called when registering components during cook/save of the level.
-			const bool bCanRemoveActor = CanRemoveActor(*Actor, Components);
-			if (UBlueprint::GetBlueprintFromClass(Actor->GetClass()) && !bCanRemoveActor)
-			{
-				continue;
-			}
-
 			FFastGeoComponentCluster* CurrentComponentCluster = LevelComponentCluster.Get();
 			TUniquePtr<FFastGeoHLOD> FastGeoHLOD;
 			if (AWorldPartitionHLOD* HLODActor = Cast<AWorldPartitionHLOD>(*Actor))
@@ -247,6 +318,8 @@ void UFastGeoWorldPartitionRuntimeCellTransformer::Transform(ULevel* InLevel)
 				CurrentComponentCluster = FastGeoHLOD.Get();
 			}
 
+			const TSet<UActorComponent*>& Components = ActorInfo.Value;
+			check(!Components.IsEmpty());
 			for (UActorComponent* Component : Components)
 			{
 				UStaticMeshComponent* StaticMeshComponent = CastChecked<UStaticMeshComponent>(Component);
@@ -259,7 +332,8 @@ void UFastGeoWorldPartitionRuntimeCellTransformer::Transform(ULevel* InLevel)
 				StaticMeshComponent->MarkAsGarbage();
 			}
 
-			if (bCanRemoveActor)
+			const bool bIsActorFullyTransformed = ActorInfo.Key;
+			if (bIsActorFullyTransformed)
 			{
 				*Actor = nullptr;
 			}
@@ -282,6 +356,13 @@ void UFastGeoWorldPartitionRuntimeCellTransformer::Transform(ULevel* InLevel)
 	}
 
 	InLevel->Actors.Remove(nullptr);
+
+	if (IsDebugMode())
+	{
+		UE_LOG(LogFastGeoStreaming, Log, TEXT("- Transformation result of Level '%s'"), *InLevel->GetPathName());
+		Stats.DumpStats(TEXT("  - "));
+		UE_LOG(LogFastGeoStreaming, Log, TEXT("------------------------------------------------------------------------"));
+	}
 }
 
 bool UFastGeoWorldPartitionRuntimeCellTransformer::IsBlueprintActorWithLogic(AActor* InActor) const
@@ -347,82 +428,210 @@ bool UFastGeoWorldPartitionRuntimeCellTransformer::IsBlueprintActorWithLogic(AAc
 	return false;
 }
 
-bool UFastGeoWorldPartitionRuntimeCellTransformer::CanTransformActor(AActor* InActor) const
+FFastGeoTransformResult UFastGeoWorldPartitionRuntimeCellTransformer::IsAllowedActorClass(AActor* InActor) const
 {
-	if (InActor->GetIsReplicated())
+	UClass* ActorClass = InActor->GetClass();
+	for (TSubclassOf<AActor> DisallowedActorClass : DisallowedActorClasses)
+	{
+		if (ActorClass->IsChildOf(DisallowedActorClass))
+		{
+			return { false, [InActor, DisallowedActorClass] { return FString::Printf(TEXT("Actor %s class is child of a disallowed class (%s)"), *InActor->GetName(), *DisallowedActorClass->GetName()); } };
+		}
+	}
+
+	for (TSubclassOf<AActor> DisallowedActorClass : DisallowedExactActorClasses)
+	{
+		if (ActorClass == DisallowedActorClass)
+		{
+			return { false, [InActor, DisallowedActorClass] { return FString::Printf(TEXT("Actor %s class is a disallowed exact class (%s)"), *InActor->GetName(), *DisallowedActorClass->GetName()); } };
+		}
+	}
+
+	// Special case where we disallow an actor class if actor is tagged 'NoFastGeo'
+	if (InActor->Tags.Contains(FastGeo::NAME_NoFastGeo))
+	{
+		return { false, [InActor] { return FString::Printf(TEXT("Actor %s is tagged 'NoFastGeo'"), *InActor->GetName()); } };
+	}
+
+	for (TSubclassOf<AActor> AllowedActorClass : AllowedActorClasses)
+	{
+		if (ActorClass->IsChildOf(AllowedActorClass))
+		{
+			return true;
+		}
+	}
+
+	for (TSubclassOf<AActor> AllowedActorClass : AllowedExactActorClasses)
+	{
+		if (ActorClass->IsChildOf(AllowedActorClass))
+		{
+			return true;
+		}
+	}
+
+	// Special case where we allow an actor class if actor is tagged 'FastGeo'
+	if (InActor->Tags.Contains(FastGeo::NAME_FastGeo))
+	{
+		return true;
+	}
+
+	return { false, [InActor, ActorClass] { return FString::Printf(TEXT("Actor %s class is an unsupported class (%s)"), *InActor->GetName(), *ActorClass->GetName()); } };
+}
+
+FFastGeoTransformResult UFastGeoWorldPartitionRuntimeCellTransformer::CanTransformActor(AActor* InActor, bool& bOutIsActorFullyTransformable, TSet<UActorComponent*>& OutTransformableComponents) const
+{
+	bOutIsActorFullyTransformable = false;
+
+	if (!IsAllowedActorClass(InActor))
 	{
 		return false;
 	}
 
+	if (InActor->GetIsReplicated())
+	{
+		return { false, [InActor] { return FString::Printf(TEXT("Actor %s is replicated"), *InActor->GetName()); } };
+	}
+
 	if (!InActor->IsRootComponentStatic())
 	{
-		return false;
+		return { false, [InActor] { return FString::Printf(TEXT("Actor %s RootComponent Mobility is not Static"), *InActor->GetName()); } };
 	}
 
 	if (InActor->IsHidden())
 	{
-		return false;
+		return { false, [InActor] { return FString::Printf(TEXT("Actor %s is hidden"), *InActor->GetName()); } };
 	}
 
 	if (InActor->IsEditorOnly())
 	{
-		return false;
+		return { false, [InActor] { return FString::Printf(TEXT("Actor %s is editor-only"), *InActor->GetName()); } };
 	}
 
-	if (InActor->Children.Num() || InActor->IsChildActor())
+	if (InActor->Children.Num())
 	{
-		return false;
+		return { false, [InActor] { return FString::Printf(TEXT("Actor %s has children"), *InActor->GetName()); } };
 	}
 
-	UClass* ActorClass = InActor->GetClass();
-	for (TSubclassOf<AActor> DisallowedActorClass : DisallowedActorClasses)
+	if (InActor->IsChildActor())
 	{
-		if (ActorClass->IsChildOf(DisallowedActorClass))
-		{
-			return false;
-		}
+		return { false, [InActor] { return FString::Printf(TEXT("Actor %s is a child actor"), *InActor->GetName()); } };
 	}
 
-	for (TSubclassOf<AActor> DisallowedActorClass : DisallowedExactActorClasses)
+	if (IsBlueprintActorWithLogic(InActor))
 	{
-		if (ActorClass == DisallowedActorClass)
+		return { false, [InActor] { return FString::Printf(TEXT("Actor %s is a Blueprint Actor with logic"), *InActor->GetName()); } };
+	}
+
+	// Gather transformable components
+	TSet<UActorComponent*> TransformableComponents;
+	InActor->ForEachComponent<UStaticMeshComponent>(false, [this, &TransformableComponents](UStaticMeshComponent* StaticMeshComponent)
+	{
+		if (CanTransformComponent(StaticMeshComponent))
 		{
-			return false;
+			TransformableComponents.Add(StaticMeshComponent);
 		}
+	});
+
+	if (TransformableComponents.IsEmpty())
+	{
+		return { false, [InActor] { return FString::Printf(TEXT("Actor %s has no transformable components"), *InActor->GetName()); } };
 	}
 
-	if (InActor->Tags.Contains(NAME_NoFastGeo))
+	// Can't convert partially a BP actors
+	// Rerun CS will be called in PIE when registering the component and
+	// also called when registering components during cook/save of the level.
+	bOutIsActorFullyTransformable = CanRemoveActor(InActor, TransformableComponents);
+	if (UBlueprint::GetBlueprintFromClass(InActor->GetClass()) && !bOutIsActorFullyTransformable)
 	{
-		return false;
+		return { false, [InActor] { return FString::Printf(TEXT("Actor %s is a Blueprint and can't be fully transformed."), *InActor->GetName()); } };
 	}
 
-	if (IsBlueprintActorWithLogic(InActor))
+	OutTransformableComponents = MoveTemp(TransformableComponents);
+	return true;
+}
+
+FFastGeoTransformResult UFastGeoWorldPartitionRuntimeCellTransformer::IsAllowedComponentClass(UActorComponent* InComponent) const
+{
+	UClass* ComponentClass = InComponent->GetClass();
+	for (TSubclassOf<UActorComponent> DisallowedComponentClass : DisallowedComponentClasses)
 	{
-		return false;
+		if (ComponentClass->IsChildOf(DisallowedComponentClass))
+		{
+			return { false, [InComponent, DisallowedComponentClass] { return FString::Printf(TEXT("Component %s class is child of a disallowed class (%s)"), FastGeo::GetComponentShortName(InComponent), *DisallowedComponentClass->GetName()); } };
+		}
 	}
 
-	for (TSubclassOf<AActor> AllowedActorClass : AllowedActorClasses)
+	for (TSubclassOf<UActorComponent> DisallowedComponentClass : DisallowedExactComponentClasses)
 	{
-		if (ActorClass->IsChildOf(AllowedActorClass))
+		if (ComponentClass == DisallowedComponentClass)
+		{
+			return { false, [InComponent, DisallowedComponentClass] { return FString::Printf(TEXT("Component %s class is a disallowed exact class (%s)"), FastGeo::GetComponentShortName(InComponent), *DisallowedComponentClass->GetName()); } };
+		}
+	}
+
+	for (TSubclassOf<UActorComponent> AllowedComponentClass : AllowedComponentClasses)
+	{
+		if (ComponentClass->IsChildOf(AllowedComponentClass))
 		{
 			return true;
 		}
 	}
 
-	for (TSubclassOf<AActor> AllowedActorClass : AllowedExactActorClasses)
+	for (TSubclassOf<UActorComponent> AllowedComponentClass : AllowedExactComponentClasses)
 	{
-		if (ActorClass->IsChildOf(AllowedActorClass))
+		if (ComponentClass == AllowedComponentClass)
 		{
 			return true;
 		}
 	}
 
-	if (InActor->Tags.Contains(NAME_FastGeo))
+	return { false, [InComponent, ComponentClass] { return FString::Printf(TEXT("Component %s class is an unsupported class (%s)"), FastGeo::GetComponentShortName(InComponent), *ComponentClass->GetName()); } };
+};
+
+FFastGeoTransformResult UFastGeoWorldPartitionRuntimeCellTransformer::CanTransformComponent(UStaticMeshComponent* InComponent) const
+{
+	if (!IsAllowedComponentClass(InComponent))
 	{
-		return true;
+		return false;
 	}
 
-	return false;
+	if (InComponent->IsEditorOnly())
+	{
+		return { false, [InComponent] { return FString::Printf(TEXT("Component %s is editor-only"), FastGeo::GetComponentShortName(InComponent)); } };
+	}
+
+	if (!InComponent->GetStaticMesh())
+	{
+		return { false, [InComponent] { return FString::Printf(TEXT("Component %s has an invalid static mesh"), FastGeo::GetComponentShortName(InComponent)); } };
+	}
+
+	if (InComponent->GetLODParentPrimitive())
+	{
+		return { false, [InComponent] { return FString::Printf(TEXT("Component %s has a valid LOD Parent Primitive"), FastGeo::GetComponentShortName(InComponent)); } };
+	}
+
+	if (InComponent->Mobility != EComponentMobility::Static)
+	{
+		return { false, [InComponent] { return FString::Printf(TEXT("Component %s Mobility is not Static"), FastGeo::GetComponentShortName(InComponent)); } };
+	}
+
+	// Make sure BodyInstance CollisionEnabled is updated first
+	InComponent->UpdateCollisionFromStaticMesh();
+	const bool bIsCollisionEnabled = FastGeo::IsCollisionEnabled(InComponent);
+
+	// If collision is enabled, only allow if async physics state creation and destruction are supported 
+	// If colliion is disabled, allow if component is visible
+	if (bIsCollisionEnabled && (!InComponent->AllowsAsyncPhysicsStateCreation() || !InComponent->AllowsAsyncPhysicsStateDestruction()))
+	{
+		return { false, [InComponent] { return FString::Printf(TEXT("Component %s has collision enabled but doesn't allow asynchronous physics state creation/destruction"), FastGeo::GetComponentShortName(InComponent)); } };
+	}
+
+	if (!bIsCollisionEnabled && !InComponent->IsVisible())
+	{
+		return { false, [InComponent] { return FString::Printf(TEXT("Component %s has no collision and is not visible"), FastGeo::GetComponentShortName(InComponent)); } };
+	}
+
+	return true;
 }
 
 bool UFastGeoWorldPartitionRuntimeCellTransformer::CanRemoveActor(AActor* InActor, TSet<UActorComponent*>& InIgnoredComponents) const
@@ -499,6 +708,43 @@ bool UFastGeoWorldPartitionRuntimeCellTransformer::CanRemoveActor(AActor* InActo
 
 	return false;
 }
+
+void UFastGeoWorldPartitionRuntimeCellTransformer::PreEditChange(FProperty* InPropertyAboutToChange)
+{
+	FastGeo::GPackageWasDirty = GetPackage()->IsDirty();
+	Super::PreEditChange(InPropertyAboutToChange);
+}
+
+void UFastGeoWorldPartitionRuntimeCellTransformer::PostEditChangeProperty(FPropertyChangedEvent& InPropertyChangedEvent)
+{
+	const FName PropertyName = InPropertyChangedEvent.GetPropertyName();
+	if ((PropertyName == GET_MEMBER_NAME_CHECKED(UFastGeoWorldPartitionRuntimeCellTransformer, bDebugMode) || 
+		 PropertyName == GET_MEMBER_NAME_CHECKED(UFastGeoWorldPartitionRuntimeCellTransformer, bDebugModeOnEditorSelection)) && !FastGeo::GPackageWasDirty)
+	{
+		GetPackage()->ClearDirtyFlag();
+	}
+}
+
+bool FFastGeoTransformResult::ShouldReport = false;
+
+FFastGeoTransformResult::FFastGeoTransformResult(bool bInSuccess, const TCHAR* FailureReason)
+	: bSuccess(bInSuccess)
+{
+	if (!bInSuccess && FailureReason && ShouldReport)
+	{
+		UE_LOG(LogFastGeoStreaming, Log, TEXT("  * Can't transform: %s"), FailureReason);
+	}
+}
+
+FFastGeoTransformResult::FFastGeoTransformResult(bool bInSuccess, TFunctionRef<FString ()> FailureReasonFunc)
+	: bSuccess(bInSuccess)
+{
+	if (!bInSuccess && ShouldReport)
+	{
+		UE_LOG(LogFastGeoStreaming, Log, TEXT("  * Can't transform: %s"), *FailureReasonFunc());
+	}
+}
+
 #endif
 
 #undef LOCTEXT_NAMESPACE
\ No newline at end of file
diff --git a/Engine/Plugins/Experimental/FastGeoStreaming/Source/FastGeoStreaming/Private/FastGeoWorldPartitionRuntimeCellTransformer.h b/Engine/Plugins/Experimental/FastGeoStreaming/Source/FastGeoStreaming/Private/FastGeoWorldPartitionRuntimeCellTransformer.h
--- a/Engine/Plugins/Experimental/FastGeoStreaming/Source/FastGeoStreaming/Private/FastGeoWorldPartitionRuntimeCellTransformer.h
+++ b/Engine/Plugins/Experimental/FastGeoStreaming/Source/FastGeoStreaming/Private/FastGeoWorldPartitionRuntimeCellTransformer.h
@@ -3,27 +3,71 @@
 #pragma once
 
 #include "WorldPartition/WorldPartitionRuntimeCellTransformer.h"
+#if WITH_EDITOR
 #include "Templates/SubclassOf.h"
+#include "HAL/IConsoleManager.h"
+#endif 
 #include "FastGeoWorldPartitionRuntimeCellTransformer.generated.h"
 
 class UActorComponent;
 
+/** Used to conditionally report the error result when not a success. */
+struct FFastGeoTransformResult
+{
+public:
+	FFastGeoTransformResult(bool bInSuccess, const TCHAR* FailureReason = nullptr);
+	FFastGeoTransformResult(bool bInSuccess, TFunctionRef<FString()> FailureReasonFunc);
+
+	explicit operator bool() const { return bSuccess; }
+
+	static bool ShouldReport;
+
+private:
+	bool bSuccess;
+};
+
 UCLASS()
 class UFastGeoWorldPartitionRuntimeCellTransformer : public UWorldPartitionRuntimeCellTransformer
 {
 	GENERATED_UCLASS_BODY()
 
+public:
+	//~ Begin UObject Interface.
+	virtual void Serialize(FArchive& Ar) override;
+	//~ End UObject Interface.
+
 #if WITH_EDITOR
 public:
+	//~ Begin UWorldPartitionRuntimeCellTransformer Interface.
 	virtual void Transform(ULevel* InLevel) override;
+	//~ End UWorldPartitionRuntimeCellTransformer Interface.
+
+	//~ Begin UObject Interface.
+	virtual void BeginDestroy() override;
+	virtual void PreEditChange(FProperty* InPropertyAboutToChange) override;
+	virtual void PostEditChangeProperty(FPropertyChangedEvent& InPropertyChangedEvent) override;
+	//~ End UObject Interface.
 
 private:
 	bool IsBlueprintActorWithLogic(AActor* InActor) const;
-	bool CanTransformActor(AActor* InActor) const;
+	FFastGeoTransformResult CanTransformActor(AActor* InActor, bool& bOutIsActorFullyTransformable, TSet<UActorComponent*>& OutTransformableComponents) const;
+	FFastGeoTransformResult IsAllowedActorClass(AActor* InActor) const;
+	FFastGeoTransformResult CanTransformComponent(UStaticMeshComponent* InComponent) const;
+	FFastGeoTransformResult IsAllowedComponentClass(UActorComponent* InComponent) const;
 	bool CanRemoveActor(AActor* InActor, TSet<UActorComponent*>& InIgnoredComponents) const;
+	void OnSelectionChanged(UObject* Object);
+	bool IsDebugMode() const;
 #endif
 
 #if WITH_EDITORONLY_DATA
+	/** When enabled, PIE will log the reason why actors/components are not transformable */
+	UPROPERTY(EditAnywhere, Transient, SkipSerialization, Category = FastGeo)
+	bool bDebugMode;
+
+	/** When enabled, Editor will log the reason why selected actors/components are not transformable */
+	UPROPERTY(EditAnywhere, Transient, SkipSerialization, Category = FastGeo)
+	bool bDebugModeOnEditorSelection;
+
 	/** Allowed actor classes (recursive) to convert to FastGeo */
 	UPROPERTY(EditAnywhere, Category = FastGeo)
 	TArray<TSubclassOf<AActor>> AllowedActorClasses;
@@ -64,4 +108,11 @@ private:
 	UPROPERTY(EditAnywhere, Category = FastGeo)
 	TArray<TSubclassOf<UActorComponent>> IgnoredRemainingExactComponentClasses;
 #endif
+
+#if WITH_EDITOR
+	static bool IsDebugModeEnabled;
+	static bool IsFastGeoEnabled;
+	static FAutoConsoleVariableRef CVarIsDebugModeEnabled;
+	static FAutoConsoleVariableRef CVarIsFastGeoEnabled;
+#endif
 };
\ No newline at end of file
diff --git a/Engine/Source/Runtime/Engine/Public/WorldPartition/WorldPartitionRuntimeCellTransformer.h b/Engine/Source/Runtime/Engine/Public/WorldPartition/WorldPartitionRuntimeCellTransformer.h
--- a/Engine/Source/Runtime/Engine/Public/WorldPartition/WorldPartitionRuntimeCellTransformer.h
+++ b/Engine/Source/Runtime/Engine/Public/WorldPartition/WorldPartitionRuntimeCellTransformer.h
@@ -4,7 +4,7 @@
 #include "CoreMinimal.h"
 #include "WorldPartitionRuntimeCellTransformer.generated.h"
 
-UCLASS(Abstract)
+UCLASS(Abstract, Within = WorldPartition)
 class ENGINE_API UWorldPartitionRuntimeCellTransformer : public UObject
 {
 	GENERATED_UCLASS_BODY()
