Fixed CIS - Mac following 41352153

diff --git a/Engine/Plugins/Experimental/FastGeoStreaming/Source/FastGeoStreaming/Private/FastGeoWorldPartitionRuntimeCellTransformer.cpp b/Engine/Plugins/Experimental/FastGeoStreaming/Source/FastGeoStreaming/Private/FastGeoWorldPartitionRuntimeCellTransformer.cpp
--- a/Engine/Plugins/Experimental/FastGeoStreaming/Source/FastGeoStreaming/Private/FastGeoWorldPartitionRuntimeCellTransformer.cpp
+++ b/Engine/Plugins/Experimental/FastGeoStreaming/Source/FastGeoStreaming/Private/FastGeoWorldPartitionRuntimeCellTransformer.cpp
@@ -56,7 +56,7 @@ namespace FastGeo
 		return Component->IsCollisionEnabled() && !Component->GetOwner()->IsA<AWorldPartitionHLOD>();
 	}
 
-	static const TCHAR* GetComponentShortName(UActorComponent* InComponent)
+	static FString GetComponentShortName(UActorComponent* InComponent)
 	{
 		TStringBuilder<256> Builder;
 		Builder += InComponent->GetName();
@@ -557,7 +557,7 @@ FFastGeoTransformResult UFastGeoWorldPartitionRuntimeCellTransformer::IsAllowedC
 	{
 		if (ComponentClass->IsChildOf(DisallowedComponentClass))
 		{
-			return { false, [InComponent, DisallowedComponentClass] { return FString::Printf(TEXT("Component %s class is child of a disallowed class (%s)"), FastGeo::GetComponentShortName(InComponent), *DisallowedComponentClass->GetName()); } };
+			return { false, [InComponent, DisallowedComponentClass] { return FString::Printf(TEXT("Component %s class is child of a disallowed class (%s)"), *FastGeo::GetComponentShortName(InComponent), *DisallowedComponentClass->GetName()); } };
 		}
 	}
 
@@ -565,7 +565,7 @@ FFastGeoTransformResult UFastGeoWorldPartitionRuntimeCellTransformer::IsAllowedC
 	{
 		if (ComponentClass == DisallowedComponentClass)
 		{
-			return { false, [InComponent, DisallowedComponentClass] { return FString::Printf(TEXT("Component %s class is a disallowed exact class (%s)"), FastGeo::GetComponentShortName(InComponent), *DisallowedComponentClass->GetName()); } };
+			return { false, [InComponent, DisallowedComponentClass] { return FString::Printf(TEXT("Component %s class is a disallowed exact class (%s)"), *FastGeo::GetComponentShortName(InComponent), *DisallowedComponentClass->GetName()); } };
 		}
 	}
 
@@ -585,7 +585,7 @@ FFastGeoTransformResult UFastGeoWorldPartitionRuntimeCellTransformer::IsAllowedC
 		}
 	}
 
-	return { false, [InComponent, ComponentClass] { return FString::Printf(TEXT("Component %s class is an unsupported class (%s)"), FastGeo::GetComponentShortName(InComponent), *ComponentClass->GetName()); } };
+	return { false, [InComponent, ComponentClass] { return FString::Printf(TEXT("Component %s class is an unsupported class (%s)"), *FastGeo::GetComponentShortName(InComponent), *ComponentClass->GetName()); } };
 };
 
 FFastGeoTransformResult UFastGeoWorldPartitionRuntimeCellTransformer::CanTransformComponent(UStaticMeshComponent* InComponent) const
@@ -597,22 +597,22 @@ FFastGeoTransformResult UFastGeoWorldPartitionRuntimeCellTransformer::CanTransfo
 
 	if (InComponent->IsEditorOnly())
 	{
-		return { false, [InComponent] { return FString::Printf(TEXT("Component %s is editor-only"), FastGeo::GetComponentShortName(InComponent)); } };
+		return { false, [InComponent] { return FString::Printf(TEXT("Component %s is editor-only"), *FastGeo::GetComponentShortName(InComponent)); } };
 	}
 
 	if (!InComponent->GetStaticMesh())
 	{
-		return { false, [InComponent] { return FString::Printf(TEXT("Component %s has an invalid static mesh"), FastGeo::GetComponentShortName(InComponent)); } };
+		return { false, [InComponent] { return FString::Printf(TEXT("Component %s has an invalid static mesh"), *FastGeo::GetComponentShortName(InComponent)); } };
 	}
 
 	if (InComponent->GetLODParentPrimitive())
 	{
-		return { false, [InComponent] { return FString::Printf(TEXT("Component %s has a valid LOD Parent Primitive"), FastGeo::GetComponentShortName(InComponent)); } };
+		return { false, [InComponent] { return FString::Printf(TEXT("Component %s has a valid LOD Parent Primitive"), *FastGeo::GetComponentShortName(InComponent)); } };
 	}
 
 	if (InComponent->Mobility != EComponentMobility::Static)
 	{
-		return { false, [InComponent] { return FString::Printf(TEXT("Component %s Mobility is not Static"), FastGeo::GetComponentShortName(InComponent)); } };
+		return { false, [InComponent] { return FString::Printf(TEXT("Component %s Mobility is not Static"), *FastGeo::GetComponentShortName(InComponent)); } };
 	}
 
 	// Make sure BodyInstance CollisionEnabled is updated first
@@ -623,12 +623,12 @@ FFastGeoTransformResult UFastGeoWorldPartitionRuntimeCellTransformer::CanTransfo
 	// If colliion is disabled, allow if component is visible
 	if (bIsCollisionEnabled && (!InComponent->AllowsAsyncPhysicsStateCreation() || !InComponent->AllowsAsyncPhysicsStateDestruction()))
 	{
-		return { false, [InComponent] { return FString::Printf(TEXT("Component %s has collision enabled but doesn't allow asynchronous physics state creation/destruction"), FastGeo::GetComponentShortName(InComponent)); } };
+		return { false, [InComponent] { return FString::Printf(TEXT("Component %s has collision enabled but doesn't allow asynchronous physics state creation/destruction"), *FastGeo::GetComponentShortName(InComponent)); } };
 	}
 
 	if (!bIsCollisionEnabled && !InComponent->IsVisible())
 	{
-		return { false, [InComponent] { return FString::Printf(TEXT("Component %s has no collision and is not visible"), FastGeo::GetComponentShortName(InComponent)); } };
+		return { false, [InComponent] { return FString::Printf(TEXT("Component %s has no collision and is not visible"), *FastGeo::GetComponentShortName(InComponent)); } };
 	}
 
 	return true;
