Removed Within = WorldPartition on UWorldPartitionRuntimeCellTransformer as there is existing content that doesn't respect that.

diff --git a/Engine/Plugins/Experimental/FastGeoStreaming/Source/FastGeoStreaming/Private/FastGeoWorldPartitionRuntimeCellTransformer.cpp b/Engine/Plugins/Experimental/FastGeoStreaming/Source/FastGeoStreaming/Private/FastGeoWorldPartitionRuntimeCellTransformer.cpp
--- a/Engine/Plugins/Experimental/FastGeoStreaming/Source/FastGeoStreaming/Private/FastGeoWorldPartitionRuntimeCellTransformer.cpp
+++ b/Engine/Plugins/Experimental/FastGeoStreaming/Source/FastGeoStreaming/Private/FastGeoWorldPartitionRuntimeCellTransformer.cpp
@@ -184,52 +184,55 @@ bool UFastGeoWorldPartitionRuntimeCellTransformer::IsDebugMode() const
 
 void UFastGeoWorldPartitionRuntimeCellTransformer::OnSelectionChanged(UObject* Object)
 {
-	if (bDebugModeOnEditorSelection && GetOuterUWorldPartition()->GetWorld() && !GetOuterUWorldPartition()->GetWorld()->IsGameWorld())
+	if (bDebugModeOnEditorSelection)
 	{
-		TArray<AActor*> SelectedActors;
-		if (USelection* Selection = Cast<USelection>(Object))
+		UWorldPartition* OuterWorldPartition = GetTypedOuter<UWorldPartition>();
+		if (OuterWorldPartition && OuterWorldPartition->GetWorld() && !OuterWorldPartition->GetWorld()->IsGameWorld())
 		{
-			for (int32 Index = 0; Index < Selection->Num(); Index++)
+			TArray<AActor*> SelectedActors;
+			if (USelection* Selection = Cast<USelection>(Object))
 			{
-				if (AActor* SelectedActor = Cast<AActor>(Selection->GetSelectedObject(Index)))
+				for (int32 Index = 0; Index < Selection->Num(); Index++)
 				{
-					UWorldPartition* WorldPartition = FWorldPartitionHelpers::GetWorldPartition(SelectedActor);
-					if (WorldPartition == GetOuterUWorldPartition())
+					if (AActor* SelectedActor = Cast<AActor>(Selection->GetSelectedObject(Index)))
 					{
-						SelectedActors.Add(SelectedActor);
+						UWorldPartition* WorldPartition = FWorldPartitionHelpers::GetWorldPartition(SelectedActor);
+						if (WorldPartition == OuterWorldPartition)
+						{
+							SelectedActors.Add(SelectedActor);
+						}
 					}
 				}
 			}
-		}
 
-		if (!SelectedActors.IsEmpty())
-		{
-			TGuardValue<bool> Guard(FFastGeoTransformResult::ShouldReport, true);
+			if (!SelectedActors.IsEmpty())
+			{
+				TGuardValue<bool> Guard(FFastGeoTransformResult::ShouldReport, true);
 
-			UE_LOG(LogFastGeoStreaming, Log, TEXT("------------------------------------------------------------------------"));
-			UE_LOG(LogFastGeoStreaming, Log, TEXT("- FastGeoStreaming Debug Mode: Transformation on %d selected actors "), SelectedActors.Num());
+				UE_LOG(LogFastGeoStreaming, Log, TEXT("------------------------------------------------------------------------"));
+				UE_LOG(LogFastGeoStreaming, Log, TEXT("- FastGeoStreaming Debug Mode: Transformation on %d selected actors "), SelectedActors.Num());
 
-			FastGeo::FTransformationStats Stats;
-			Stats.TotalActorCount = SelectedActors.Num();
-			for (AActor* Actor : SelectedActors)
-			{
-				bool bIsActorFullyTransformable;
-				TSet<UActorComponent*> TransformableComponents;
-				if (CanTransformActor(Actor, bIsActorFullyTransformable, TransformableComponents))
+				FastGeo::FTransformationStats Stats;
+				Stats.TotalActorCount = SelectedActors.Num();
+				for (AActor* Actor : SelectedActors)
 				{
-					Stats.FullyTransformableActorCount += bIsActorFullyTransformable ? 1 : 0;
-					Stats.PartiallyTransformableActorCount += !bIsActorFullyTransformable ? 1 : 0;
-					Stats.TransformedComponentCount += TransformableComponents.Num();
+					bool bIsActorFullyTransformable;
+					TSet<UActorComponent*> TransformableComponents;
+					if (CanTransformActor(Actor, bIsActorFullyTransformable, TransformableComponents))
+					{
+						Stats.FullyTransformableActorCount += bIsActorFullyTransformable ? 1 : 0;
+						Stats.PartiallyTransformableActorCount += !bIsActorFullyTransformable ? 1 : 0;
+						Stats.TransformedComponentCount += TransformableComponents.Num();
+					}
+					Stats.TotalComponentCount += Actor->GetComponents().Num();
 				}
-				Stats.TotalComponentCount += Actor->GetComponents().Num();
-			}
 
-			Stats.DumpStats(TEXT("  - "));
-			UE_LOG(LogFastGeoStreaming, Log, TEXT("------------------------------------------------------------------------"));
+				Stats.DumpStats(TEXT("  - "));
+				UE_LOG(LogFastGeoStreaming, Log, TEXT("------------------------------------------------------------------------"));
+			}
 		}
 	}
 }
-
 void UFastGeoWorldPartitionRuntimeCellTransformer::Transform(ULevel* InLevel)
 {
 	TGuardValue<bool> Guard(FFastGeoTransformResult::ShouldReport, IsDebugMode());
diff --git a/Engine/Source/Runtime/Engine/Public/WorldPartition/WorldPartitionRuntimeCellTransformer.h b/Engine/Source/Runtime/Engine/Public/WorldPartition/WorldPartitionRuntimeCellTransformer.h
--- a/Engine/Source/Runtime/Engine/Public/WorldPartition/WorldPartitionRuntimeCellTransformer.h
+++ b/Engine/Source/Runtime/Engine/Public/WorldPartition/WorldPartitionRuntimeCellTransformer.h
@@ -4,7 +4,7 @@
 #include "CoreMinimal.h"
 #include "WorldPartitionRuntimeCellTransformer.generated.h"
 
-UCLASS(Abstract, Within = WorldPartition)
+UCLASS(Abstract)
 class ENGINE_API UWorldPartitionRuntimeCellTransformer : public UObject
 {
 	GENERATED_UCLASS_BODY()
