FastGeo : Added AWorldPartitionAutoInstancedActor in default AllowedActorClasses.

diff --git a/Engine/Plugins/Experimental/FastGeoStreaming/Source/FastGeoStreaming/Private/FastGeoWorldPartitionRuntimeCellTransformer.cpp b/Engine/Plugins/Experimental/FastGeoStreaming/Source/FastGeoStreaming/Private/FastGeoWorldPartitionRuntimeCellTransformer.cpp
--- a/Engine/Plugins/Experimental/FastGeoStreaming/Source/FastGeoStreaming/Private/FastGeoWorldPartitionRuntimeCellTransformer.cpp
+++ b/Engine/Plugins/Experimental/FastGeoStreaming/Source/FastGeoStreaming/Private/FastGeoWorldPartitionRuntimeCellTransformer.cpp
@@ -22,6 +22,7 @@
 #include "GameFramework/Info.h"
 #include "WorldPartition/WorldPartition.h"
 #include "WorldPartition/WorldPartitionHelpers.h"
+#include "WorldPartition/WorldPartitionRuntimeCellTransformerISM.h"
 #include "Selection.h"
 #include "FastGeoContainer.h"
 #include "FastGeoComponentCluster.h"
@@ -123,6 +124,7 @@ UFastGeoWorldPartitionRuntimeCellTransformer::UFastGeoWorldPartitionRuntimeCellT
 	AllowedActorClasses.Add(APackedLevelActor::StaticClass());
 	AllowedActorClasses.Add(AStaticMeshActor::StaticClass());
 	AllowedActorClasses.Add(AWorldPartitionHLOD::StaticClass());
+	AllowedActorClasses.Add(AWorldPartitionAutoInstancedActor::StaticClass());
 	AllowedComponentClasses.Add(UStaticMeshComponent::StaticClass());
 
 	// Disallowed actor/component classes
