FastGeo: Debug Mode on editor selection now supports Level Instances.

diff --git a/Engine/Plugins/Experimental/FastGeoStreaming/Source/FastGeoStreaming/Private/FastGeoWorldPartitionRuntimeCellTransformer.cpp b/Engine/Plugins/Experimental/FastGeoStreaming/Source/FastGeoStreaming/Private/FastGeoWorldPartitionRuntimeCellTransformer.cpp
--- a/Engine/Plugins/Experimental/FastGeoStreaming/Source/FastGeoStreaming/Private/FastGeoWorldPartitionRuntimeCellTransformer.cpp
+++ b/Engine/Plugins/Experimental/FastGeoStreaming/Source/FastGeoStreaming/Private/FastGeoWorldPartitionRuntimeCellTransformer.cpp
@@ -22,7 +22,13 @@
 #include "GameFramework/Info.h"
 #include "WorldPartition/WorldPartition.h"
 #include "WorldPartition/WorldPartitionHelpers.h"
+#include "WorldPartition/DataLayer/WorldDataLayers.h"
 #include "WorldPartition/WorldPartitionRuntimeCellTransformerISM.h"
+#include "LevelInstance/LevelInstanceInterface.h"
+#include "LevelInstance/LevelInstanceSubsystem.h"
+#include "LevelInstance/LevelInstanceEditorInstanceActor.h"
+#include "LevelInstance/LevelInstanceEditorPivotInterface.h"
+#include "ActorEditorUtils.h"
 #include "Selection.h"
 #include "FastGeoContainer.h"
 #include "FastGeoComponentCluster.h"
@@ -193,20 +199,41 @@ void UFastGeoWorldPartitionRuntimeCellTransformer::OnSelectionChanged(UObject* O
 	if (bDebugModeOnEditorSelection)
 	{
 		UWorldPartition* OuterWorldPartition = GetTypedOuter<UWorldPartition>();
-		if (OuterWorldPartition && OuterWorldPartition->GetWorld() && !OuterWorldPartition->GetWorld()->IsGameWorld())
+		UWorld* World = OuterWorldPartition ? OuterWorldPartition->GetWorld() : nullptr;
+		if (World && !World->IsGameWorld())
 		{
 			TArray<AActor*> SelectedActors;
+
+			TFunction<void(AActor* InActor)> AddActorToSelection = [this, World, &SelectedActors, &AddActorToSelection](AActor* InActor)
+			{
+				if ((InActor->GetWorld() == World) && !CanAlwaysIgnoreActor(InActor))
+				{
+					const ILevelInstanceInterface* LevelInstance = Cast<ILevelInstanceInterface>(InActor);
+					if (LevelInstance && LevelInstance->GetDesiredRuntimeBehavior() == ELevelInstanceRuntimeBehavior::Partitioned && InActor->GetWorld())
+					{
+						if (const ULevelInstanceSubsystem* LevelInstanceSubsystem = InActor->GetWorld()->GetSubsystem<ULevelInstanceSubsystem>())
+						{
+							LevelInstanceSubsystem->ForEachActorInLevelInstance(LevelInstance, [&AddActorToSelection](AActor* Actor)
+							{
+								AddActorToSelection(Actor);
+								return true;
+							});
+						}
+					}
+					else
+					{
+						SelectedActors.Add(InActor);
+					}
+				}
+			};
+
 			if (USelection* Selection = Cast<USelection>(Object))
 			{
 				for (int32 Index = 0; Index < Selection->Num(); Index++)
 				{
 					if (AActor* SelectedActor = Cast<AActor>(Selection->GetSelectedObject(Index)))
 					{
-						UWorldPartition* WorldPartition = FWorldPartitionHelpers::GetWorldPartition(SelectedActor);
-						if (WorldPartition == OuterWorldPartition)
-						{
-							SelectedActors.Add(SelectedActor);
-						}
+						AddActorToSelection(SelectedActor);
 					}
 				}
 			}
@@ -239,6 +266,16 @@ void UFastGeoWorldPartitionRuntimeCellTransformer::OnSelectionChanged(UObject* O
 		}
 	}
 }
+
+bool UFastGeoWorldPartitionRuntimeCellTransformer::CanAlwaysIgnoreActor(AActor* InActor) const
+{
+	return InActor->IsA<AWorldSettings>() ||
+		InActor->IsA<AWorldDataLayers>() ||
+		InActor->IsA<ALevelInstanceEditorInstanceActor>() ||
+		InActor->Implements<ULevelInstanceEditorPivotInterface>() ||
+		FActorEditorUtils::IsABuilderBrush(InActor);
+}
+
 void UFastGeoWorldPartitionRuntimeCellTransformer::Transform(ULevel* InLevel)
 {
 	TGuardValue<bool> Guard(FFastGeoTransformResult::ShouldReport, IsDebugMode());
@@ -267,8 +304,7 @@ void UFastGeoWorldPartitionRuntimeCellTransformer::Transform(ULevel* InLevel)
 	TMap<TObjectPtr<AActor>*, TPair<bool, TSet<UActorComponent*>>> ActorComponents;
 	for (TObjectPtr<AActor>& Actor : InLevel->Actors)
 	{
-		// Always skip the AWorldSettings of a Level (don't even count it)
-		if (IsValid(Actor) && !Actor->IsA<AWorldSettings>())
+		if (IsValid(Actor) && !CanAlwaysIgnoreActor(Actor))
 		{
 			Stats.TotalActorCount++;
 			Stats.TotalComponentCount += Actor->GetComponents().Num();
diff --git a/Engine/Plugins/Experimental/FastGeoStreaming/Source/FastGeoStreaming/Private/FastGeoWorldPartitionRuntimeCellTransformer.h b/Engine/Plugins/Experimental/FastGeoStreaming/Source/FastGeoStreaming/Private/FastGeoWorldPartitionRuntimeCellTransformer.h
--- a/Engine/Plugins/Experimental/FastGeoStreaming/Source/FastGeoStreaming/Private/FastGeoWorldPartitionRuntimeCellTransformer.h
+++ b/Engine/Plugins/Experimental/FastGeoStreaming/Source/FastGeoStreaming/Private/FastGeoWorldPartitionRuntimeCellTransformer.h
@@ -55,6 +55,7 @@ private:
 	FFastGeoTransformResult CanTransformComponent(UPrimitiveComponent* InComponent) const;
 	FFastGeoTransformResult IsAllowedComponentClass(UActorComponent* InComponent) const;
 	bool CanRemoveActor(AActor* InActor, TSet<UActorComponent*>& InIgnoredComponents) const;
+	bool CanAlwaysIgnoreActor(AActor* InActor) const;
 	void OnSelectionChanged(UObject* Object);
 	bool IsDebugMode() const;
 #endif
