FastGeo: Fixed displayed component name when using debug mode

diff --git a/Engine/Plugins/Experimental/FastGeoStreaming/Source/FastGeoStreaming/Private/FastGeoWorldPartitionRuntimeCellTransformer.cpp b/Engine/Plugins/Experimental/FastGeoStreaming/Source/FastGeoStreaming/Private/FastGeoWorldPartitionRuntimeCellTransformer.cpp
--- a/Engine/Plugins/Experimental/FastGeoStreaming/Source/FastGeoStreaming/Private/FastGeoWorldPartitionRuntimeCellTransformer.cpp
+++ b/Engine/Plugins/Experimental/FastGeoStreaming/Source/FastGeoStreaming/Private/FastGeoWorldPartitionRuntimeCellTransformer.cpp
@@ -70,9 +70,9 @@ namespace FastGeo
 	static FString GetComponentShortName(UActorComponent* InComponent)
 	{
 		TStringBuilder<256> Builder;
-		Builder += InComponent->GetName();
-		Builder += TEXT(".");
 		Builder += InComponent->GetOwner()->GetName();
+		Builder += TEXT(".");
+		Builder += InComponent->GetName();
 		return *Builder;
 	}
 
