FastGeo: Fixed crash when reloading an InstancedSkinnedMeshComponent

diff --git a/Engine/Plugins/Experimental/FastGeoStreaming/Source/FastGeoStreaming/Private/FastGeoSkinnedMeshComponent.cpp b/Engine/Plugins/Experimental/FastGeoStreaming/Source/FastGeoStreaming/Private/FastGeoSkinnedMeshComponent.cpp
--- a/Engine/Plugins/Experimental/FastGeoStreaming/Source/FastGeoStreaming/Private/FastGeoSkinnedMeshComponent.cpp
+++ b/Engine/Plugins/Experimental/FastGeoStreaming/Source/FastGeoStreaming/Private/FastGeoSkinnedMeshComponent.cpp
@@ -115,13 +115,15 @@ void FFastGeoSkinnedMeshComponentBase::InitializeSceneProxyDescDynamicProperties
 	// Destroy any existing mesh object
 	DestroyMeshObject();
 
+	// Create mesh object
+	MeshObject = CreateMeshObject();
+
 	// Initialize non-serialized properties
 	FSkinnedMeshSceneProxyDesc& SceneProxyDesc = GetSkinnedMeshSceneProxyDesc();
 	SceneProxyDesc.OverrideMaterials = OverrideMaterials;
 	SceneProxyDesc.SkinCacheUsage = SkinCacheUsage;
 	SceneProxyDesc.MaterialRelevance = GetMaterialRelevance(GetScene()->GetFeatureLevel());
-	SceneProxyDesc.MeshObject = CreateMeshObject();
-	MeshObject = SceneProxyDesc.MeshObject;
+	SceneProxyDesc.MeshObject = MeshObject;
 }
 
 FPrimitiveSceneProxy* FFastGeoSkinnedMeshComponentBase::CreateSceneProxy(ESceneProxyCreationError* OutError)
diff --git a/Engine/Plugins/Experimental/FastGeoStreaming/Source/FastGeoStreaming/Private/FastGeoSkinnedMeshComponent.h b/Engine/Plugins/Experimental/FastGeoStreaming/Source/FastGeoStreaming/Private/FastGeoSkinnedMeshComponent.h
--- a/Engine/Plugins/Experimental/FastGeoStreaming/Source/FastGeoStreaming/Private/FastGeoSkinnedMeshComponent.h
+++ b/Engine/Plugins/Experimental/FastGeoStreaming/Source/FastGeoStreaming/Private/FastGeoSkinnedMeshComponent.h
@@ -68,8 +68,8 @@ protected:
 
 	FBoxSphereBounds CalcMeshBound(const FVector3f& InRootOffset, bool bInUsePhysicsAsset, const FTransform& InLocalToWorld) const;
 	USkinnedAsset* GetSkinnedAsset() const { return GetSkinnedMeshSceneProxyDesc().GetSkinnedAsset(); }
-	FSkeletalMeshObject* GetMeshObject() { return GetSkinnedMeshSceneProxyDesc().MeshObject; }
-	const FSkeletalMeshObject* GetMeshObject() const { return GetSkinnedMeshSceneProxyDesc().MeshObject; }
+	FSkeletalMeshObject* GetMeshObject() { return MeshObject; }
+	const FSkeletalMeshObject* GetMeshObject() const { return MeshObject; }
 	void DestroyMeshObject();
 
 	//~ Begin Unsupported
