FastGeo: Fixed invalid condition in FFastGeoInstancedStaticMeshComponent::CreateSceneProxy that prevented some primitive from being rendered

diff --git a/Engine/Plugins/Experimental/FastGeoStreaming/Source/FastGeoStreaming/Private/FastGeoInstancedStaticMeshComponent.cpp b/Engine/Plugins/Experimental/FastGeoStreaming/Source/FastGeoStreaming/Private/FastGeoInstancedStaticMeshComponent.cpp
--- a/Engine/Plugins/Experimental/FastGeoStreaming/Source/FastGeoStreaming/Private/FastGeoInstancedStaticMeshComponent.cpp
+++ b/Engine/Plugins/Experimental/FastGeoStreaming/Source/FastGeoStreaming/Private/FastGeoInstancedStaticMeshComponent.cpp
@@ -197,7 +197,7 @@ FPrimitiveSceneProxy* FFastGeoInstancedStaticMeshComponent::CreateSceneProxy(ESc
 	check(SceneProxyDesc.ComponentId == GetPrimitiveSceneId());
 
 	// Verify that the mesh is valid before using it.
-	const bool bMeshIsValid = PerInstanceSMData.Num() > 0 && StaticMesh && StaticMesh->IsCompiling() == false && StaticMesh->HasValidRenderData();
+	const bool bMeshIsValid = PerInstanceSMData.Num() > 0 && StaticMesh && !StaticMesh->IsCompiling() && (StaticMesh->HasValidRenderData() || StaticMesh->HasValidNaniteData());
 	if (!bMeshIsValid)
 	{
 		if (OutError)
