Fixed crash at cooking caused by FFastGeoInstancedSkinnedMeshComponent not property initializing its FFastGeoElementType.

diff --git a/Engine/Plugins/Experimental/FastGeoStreaming/Source/FastGeoStreaming/Private/FastGeoInstancedSkinnedMeshComponent.cpp b/Engine/Plugins/Experimental/FastGeoStreaming/Source/FastGeoStreaming/Private/FastGeoInstancedSkinnedMeshComponent.cpp
--- a/Engine/Plugins/Experimental/FastGeoStreaming/Source/FastGeoStreaming/Private/FastGeoInstancedSkinnedMeshComponent.cpp
+++ b/Engine/Plugins/Experimental/FastGeoStreaming/Source/FastGeoStreaming/Private/FastGeoInstancedSkinnedMeshComponent.cpp
@@ -12,7 +12,7 @@
 #include "SkeletalRenderPublic.h"
 #include "SkinnedMeshComponentHelper.h"
 
-const FFastGeoElementType FFastGeoInstancedSkinnedMeshComponent::Type(&FFastGeoStaticMeshComponentBase::Type);
+const FFastGeoElementType FFastGeoInstancedSkinnedMeshComponent::Type(&FFastGeoSkinnedMeshComponentBase::Type);
 
 FFastGeoInstancedSkinnedMeshComponent::FFastGeoInstancedSkinnedMeshComponent(int32 InComponentIndex, FFastGeoElementType InType)
 	: Super(InComponentIndex, InType)
