FastGeo - Unified logic used to exclude components from the fast geo transform

diff --git a/Engine/Plugins/Experimental/FastGeoStreaming/Source/FastGeoStreaming/Private/FastGeoWorldPartitionRuntimeCellTransformer.cpp b/Engine/Plugins/Experimental/FastGeoStreaming/Source/FastGeoStreaming/Private/FastGeoWorldPartitionRuntimeCellTransformer.cpp
--- a/Engine/Plugins/Experimental/FastGeoStreaming/Source/FastGeoStreaming/Private/FastGeoWorldPartitionRuntimeCellTransformer.cpp
+++ b/Engine/Plugins/Experimental/FastGeoStreaming/Source/FastGeoStreaming/Private/FastGeoWorldPartitionRuntimeCellTransformer.cpp
@@ -148,6 +148,9 @@ UFastGeoWorldPartitionRuntimeCellTransformer::UFastGeoWorldPartitionRuntimeCellT
 	// Ignored component classes
 	IgnoredRemainingComponentClasses.Add(ULevelInstanceComponent::StaticClass());
 	IgnoredRemainingComponentClasses.Add(UActorTextureStreamingBuildDataComponent::StaticClass());
+
+	// Ignored remaining exact component classes
+	IgnoredRemainingExactComponentClasses.Add(USceneComponent::StaticClass());
 #endif
 }
 
@@ -250,8 +253,9 @@ void UFastGeoWorldPartitionRuntimeCellTransformer::OnSelectionChanged(UObject* O
 				for (AActor* Actor : SelectedActors)
 				{
 					bool bIsActorFullyTransformable;
-					TSet<UActorComponent*> TransformableComponents;
-					if (CanTransformActor(Actor, bIsActorFullyTransformable, TransformableComponents))
+					TArray<UActorComponent*> TransformableComponents;
+					FFastGeoTransformResult ActorTransformResult = CanTransformActor(Actor, bIsActorFullyTransformable, TransformableComponents);
+					if (ActorTransformResult.GetResult() == EFastGeoTransform::Allow)
 					{
 						Stats.FullyTransformableActorCount += bIsActorFullyTransformable ? 1 : 0;
 						Stats.PartiallyTransformableActorCount += !bIsActorFullyTransformable ? 1 : 0;
@@ -301,7 +305,7 @@ void UFastGeoWorldPartitionRuntimeCellTransformer::Transform(ULevel* InLevel)
 	}
 
 	FastGeo::FTransformationStats Stats;
-	TMap<TObjectPtr<AActor>*, TPair<bool, TSet<UActorComponent*>>> ActorComponents;
+	TMap<TObjectPtr<AActor>*, TPair<bool, TArray<UActorComponent*>>> ActorComponents;
 	for (TObjectPtr<AActor>& Actor : InLevel->Actors)
 	{
 		if (IsValid(Actor) && !CanAlwaysIgnoreActor(Actor))
@@ -310,8 +314,9 @@ void UFastGeoWorldPartitionRuntimeCellTransformer::Transform(ULevel* InLevel)
 			Stats.TotalComponentCount += Actor->GetComponents().Num();
 
 			bool bIsActorFullyTransformable;
-			TSet<UActorComponent*> TransformableComponents;
-			if (CanTransformActor(Actor, bIsActorFullyTransformable, TransformableComponents))
+			TArray<UActorComponent*> TransformableComponents;
+			FFastGeoTransformResult ActorTransformResult = CanTransformActor(Actor, bIsActorFullyTransformable, TransformableComponents);
+			if (ActorTransformResult.GetResult() == EFastGeoTransform::Allow)
 			{
 				Stats.FullyTransformableActorCount += bIsActorFullyTransformable ? 1 : 0;
 				Stats.PartiallyTransformableActorCount += !bIsActorFullyTransformable ? 1 : 0;
@@ -365,7 +370,7 @@ void UFastGeoWorldPartitionRuntimeCellTransformer::Transform(ULevel* InLevel)
 				CurrentComponentCluster = FastGeoHLOD.Get();
 			}
 
-			const TSet<UActorComponent*>& Components = ActorInfo.Value;
+			const TArray<UActorComponent*>& Components = ActorInfo.Value;
 			check(!Components.IsEmpty());
 			for (UActorComponent* Component : Components)
 			{
@@ -481,7 +486,7 @@ FFastGeoTransformResult UFastGeoWorldPartitionRuntimeCellTransformer::IsAllowedA
 	{
 		if (ActorClass->IsChildOf(DisallowedActorClass))
 		{
-			return { false, [InActor, DisallowedActorClass] { return FString::Printf(TEXT("Actor %s class is child of a disallowed class (%s)"), *InActor->GetName(), *DisallowedActorClass->GetName()); } };
+			return { EFastGeoTransform::Reject, [InActor, DisallowedActorClass] { return FString::Printf(TEXT("Actor %s class is child of a disallowed class (%s)"), *InActor->GetName(), *DisallowedActorClass->GetName()); } };
 		}
 	}
 
@@ -489,21 +494,21 @@ FFastGeoTransformResult UFastGeoWorldPartitionRuntimeCellTransformer::IsAllowedA
 	{
 		if (ActorClass == DisallowedActorClass)
 		{
-			return { false, [InActor, DisallowedActorClass] { return FString::Printf(TEXT("Actor %s class is a disallowed exact class (%s)"), *InActor->GetName(), *DisallowedActorClass->GetName()); } };
+			return { EFastGeoTransform::Reject, [InActor, DisallowedActorClass] { return FString::Printf(TEXT("Actor %s class is a disallowed exact class (%s)"), *InActor->GetName(), *DisallowedActorClass->GetName()); } };
 		}
 	}
 
 	// Special case where we disallow an actor class if actor is tagged 'NoFastGeo'
 	if (InActor->Tags.Contains(FastGeo::NAME_NoFastGeo))
 	{
-		return { false, [InActor] { return FString::Printf(TEXT("Actor %s is tagged 'NoFastGeo'"), *InActor->GetName()); } };
+		return { EFastGeoTransform::Reject, [InActor] { return FString::Printf(TEXT("Actor %s is tagged 'NoFastGeo'"), *InActor->GetName()); } };
 	}
 
 	for (TSubclassOf<AActor> AllowedActorClass : AllowedActorClasses)
 	{
 		if (ActorClass->IsChildOf(AllowedActorClass))
 		{
-			return true;
+			return EFastGeoTransform::Allow;
 		}
 	}
 
@@ -511,89 +516,90 @@ FFastGeoTransformResult UFastGeoWorldPartitionRuntimeCellTransformer::IsAllowedA
 	{
 		if (ActorClass->IsChildOf(AllowedActorClass))
 		{
-			return true;
+			return EFastGeoTransform::Allow;
 		}
 	}
 
 	// Special case where we allow an actor class if actor is tagged 'FastGeo'
 	if (InActor->Tags.Contains(FastGeo::NAME_FastGeo))
 	{
-		return true;
+		return EFastGeoTransform::Allow;
 	}
 
-	return { false, [InActor, ActorClass] { return FString::Printf(TEXT("Actor %s class is an unsupported class (%s)"), *InActor->GetName(), *ActorClass->GetName()); } };
+	return { EFastGeoTransform::Reject, [InActor, ActorClass] { return FString::Printf(TEXT("Actor %s class is an unsupported class (%s)"), *InActor->GetName(), *ActorClass->GetName()); } };
 }
 
-FFastGeoTransformResult UFastGeoWorldPartitionRuntimeCellTransformer::CanTransformActor(AActor* InActor, bool& bOutIsActorFullyTransformable, TSet<UActorComponent*>& OutTransformableComponents) const
+FFastGeoTransformResult UFastGeoWorldPartitionRuntimeCellTransformer::CanTransformActor(AActor* InActor, bool& bOutIsActorFullyTransformable, TArray<UActorComponent*>& OutTransformableComponents) const
 {
 	bOutIsActorFullyTransformable = false;
 
-	if (!IsAllowedActorClass(InActor))
+	FFastGeoTransformResult AllowedActorClassResult = IsAllowedActorClass(InActor);
+	if (AllowedActorClassResult.GetResult() != EFastGeoTransform::Allow)
 	{
-		return false;
+		return AllowedActorClassResult;
 	}
 
 	if (InActor->GetIsReplicated())
 	{
-		return { false, [InActor] { return FString::Printf(TEXT("Actor %s is replicated"), *InActor->GetName()); } };
+		return { EFastGeoTransform::Reject, [InActor] { return FString::Printf(TEXT("Actor %s is replicated"), *InActor->GetName()); } };
 	}
 
 	if (!InActor->IsRootComponentStatic())
 	{
-		return { false, [InActor] { return FString::Printf(TEXT("Actor %s RootComponent Mobility is not Static"), *InActor->GetName()); } };
+		return { EFastGeoTransform::Reject, [InActor] { return FString::Printf(TEXT("Actor %s RootComponent Mobility is not Static"), *InActor->GetName()); } };
 	}
 
 	if (InActor->IsHidden())
 	{
-		return { false, [InActor] { return FString::Printf(TEXT("Actor %s is hidden"), *InActor->GetName()); } };
+		return { EFastGeoTransform::Reject, [InActor] { return FString::Printf(TEXT("Actor %s is hidden"), *InActor->GetName()); } };
 	}
 
 	if (InActor->IsEditorOnly())
 	{
-		return { false, [InActor] { return FString::Printf(TEXT("Actor %s is editor-only"), *InActor->GetName()); } };
+		return { EFastGeoTransform::Reject, [InActor] { return FString::Printf(TEXT("Actor %s is editor-only"), *InActor->GetName()); } };
 	}
 
 	if (InActor->Children.Num())
 	{
-		return { false, [InActor] { return FString::Printf(TEXT("Actor %s has children"), *InActor->GetName()); } };
+		return { EFastGeoTransform::Reject, [InActor] { return FString::Printf(TEXT("Actor %s has children"), *InActor->GetName()); } };
 	}
 
 	if (InActor->IsChildActor())
 	{
-		return { false, [InActor] { return FString::Printf(TEXT("Actor %s is a child actor"), *InActor->GetName()); } };
+		return { EFastGeoTransform::Reject, [InActor] { return FString::Printf(TEXT("Actor %s is a child actor"), *InActor->GetName()); } };
 	}
 
 	if (IsBlueprintActorWithLogic(InActor))
 	{
-		return { false, [InActor] { return FString::Printf(TEXT("Actor %s is a Blueprint Actor with logic"), *InActor->GetName()); } };
+		return { EFastGeoTransform::Reject, [InActor] { return FString::Printf(TEXT("Actor %s is a Blueprint Actor with logic"), *InActor->GetName()); } };
 	}
 
 	// Gather transformable components
-	TSet<UActorComponent*> TransformableComponents;
-	InActor->ForEachComponent<UPrimitiveComponent>(false, [this, &TransformableComponents](UPrimitiveComponent* PrimitiveComponent)
+	TArray<UActorComponent*> TransformResults[EnumToIndex(EFastGeoTransform::MAX)];
+	InActor->ForEachComponent<UPrimitiveComponent>(false, [this, &TransformResults](UPrimitiveComponent* PrimitiveComponent)
 	{
-		if (CanTransformComponent(PrimitiveComponent))
-		{
-			TransformableComponents.Add(PrimitiveComponent);
-		}
+		FFastGeoTransformResult TransformComponentResult = CanTransformComponent(PrimitiveComponent);
+		TransformResults[TransformComponentResult.GetResultIndex()].Add(PrimitiveComponent);
 	});
 
-	if (TransformableComponents.IsEmpty())
+	if (TransformResults[EnumToIndex(EFastGeoTransform::Allow)].IsEmpty())
 	{
-		return { false, [InActor] { return FString::Printf(TEXT("Actor %s has no transformable components"), *InActor->GetName()); } };
+		return { EFastGeoTransform::Reject, [InActor] { return FString::Printf(TEXT("Actor %s has no transformable components"), *InActor->GetName()); } };
 	}
 
+	// If actor contains only discardable or transformable components, we can actually get rid of it
+	bOutIsActorFullyTransformable = TransformResults[EnumToIndex(EFastGeoTransform::Reject)].IsEmpty();
+
 	// Can't convert partially a BP actors
 	// Rerun CS will be called in PIE when registering the component and
 	// also called when registering components during cook/save of the level.
-	bOutIsActorFullyTransformable = CanRemoveActor(InActor, TransformableComponents);
 	if (UBlueprint::GetBlueprintFromClass(InActor->GetClass()) && !bOutIsActorFullyTransformable)
 	{
-		return { false, [InActor] { return FString::Printf(TEXT("Actor %s is a Blueprint and can't be fully transformed."), *InActor->GetName()); } };
+		return { EFastGeoTransform::Reject, [InActor] { return FString::Printf(TEXT("Actor %s is a Blueprint and can't be fully transformed."), *InActor->GetName()); } };
 	}
 
-	OutTransformableComponents = MoveTemp(TransformableComponents);
-	return true;
+	OutTransformableComponents = MoveTemp(TransformResults[EnumToIndex(EFastGeoTransform::Allow)]);
+	return EFastGeoTransform::Allow;
 }
 
 FFastGeoTransformResult UFastGeoWorldPartitionRuntimeCellTransformer::IsAllowedComponentClass(UActorComponent* InComponent) const
@@ -603,7 +609,7 @@ FFastGeoTransformResult UFastGeoWorldPartitionRuntimeCellTransformer::IsAllowedC
 	{
 		if (ComponentClass->IsChildOf(DisallowedComponentClass))
 		{
-			return { false, [InComponent, DisallowedComponentClass] { return FString::Printf(TEXT("Component %s class is child of a disallowed class (%s)"), *FastGeo::GetComponentShortName(InComponent), *DisallowedComponentClass->GetName()); } };
+			return { EFastGeoTransform::Reject, [InComponent, DisallowedComponentClass] { return FString::Printf(TEXT("Component %s class is child of a disallowed class (%s)"), *FastGeo::GetComponentShortName(InComponent), *DisallowedComponentClass->GetName()); } };
 		}
 	}
 
@@ -611,7 +617,7 @@ FFastGeoTransformResult UFastGeoWorldPartitionRuntimeCellTransformer::IsAllowedC
 	{
 		if (ComponentClass == DisallowedComponentClass)
 		{
-			return { false, [InComponent, DisallowedComponentClass] { return FString::Printf(TEXT("Component %s class is a disallowed exact class (%s)"), *FastGeo::GetComponentShortName(InComponent), *DisallowedComponentClass->GetName()); } };
+			return { EFastGeoTransform::Reject, [InComponent, DisallowedComponentClass] { return FString::Printf(TEXT("Component %s class is a disallowed exact class (%s)"), *FastGeo::GetComponentShortName(InComponent), *DisallowedComponentClass->GetName()); } };
 		}
 	}
 
@@ -619,7 +625,7 @@ FFastGeoTransformResult UFastGeoWorldPartitionRuntimeCellTransformer::IsAllowedC
 	{
 		if (ComponentClass->IsChildOf(AllowedComponentClass))
 		{
-			return true;
+			return EFastGeoTransform::Allow;
 		}
 	}
 
@@ -627,33 +633,34 @@ FFastGeoTransformResult UFastGeoWorldPartitionRuntimeCellTransformer::IsAllowedC
 	{
 		if (ComponentClass == AllowedComponentClass)
 		{
-			return true;
+			return EFastGeoTransform::Allow;
 		}
 	}
 
-	return { false, [InComponent, ComponentClass] { return FString::Printf(TEXT("Component %s class is an unsupported class (%s)"), *FastGeo::GetComponentShortName(InComponent), *ComponentClass->GetName()); } };
+	return { EFastGeoTransform::Reject, [InComponent, ComponentClass] { return FString::Printf(TEXT("Component %s class is an unsupported class (%s)"), *FastGeo::GetComponentShortName(InComponent), *ComponentClass->GetName()); } };
 };
 
 FFastGeoTransformResult UFastGeoWorldPartitionRuntimeCellTransformer::CanTransformComponent(UPrimitiveComponent* InComponent) const
 {
-	if (!IsAllowedComponentClass(InComponent))
+	FFastGeoTransformResult AllowedComponentClassResult = IsAllowedComponentClass(InComponent);
+	if (AllowedComponentClassResult.GetResult() != EFastGeoTransform::Allow)
 	{
-		return false;
+		return AllowedComponentClassResult;
 	}
 
 	if (InComponent->IsEditorOnly())
 	{
-		return { false, [InComponent] { return FString::Printf(TEXT("Component %s is editor-only"), *FastGeo::GetComponentShortName(InComponent)); } };
+		return { EFastGeoTransform::Discard, [InComponent] { return FString::Printf(TEXT("Component %s is editor-only"), *FastGeo::GetComponentShortName(InComponent)); } };
 	}
 
 	if (InComponent->GetLODParentPrimitive())
 	{
-		return { false, [InComponent] { return FString::Printf(TEXT("Component %s has a valid LOD Parent Primitive"), *FastGeo::GetComponentShortName(InComponent)); } };
+		return { EFastGeoTransform::Reject, [InComponent] { return FString::Printf(TEXT("Component %s has a valid LOD Parent Primitive"), *FastGeo::GetComponentShortName(InComponent)); } };
 	}
 
 	if (InComponent->Mobility != EComponentMobility::Static)
 	{
-		return { false, [InComponent] { return FString::Printf(TEXT("Component %s Mobility is not Static"), *FastGeo::GetComponentShortName(InComponent)); } };
+		return { EFastGeoTransform::Reject, [InComponent] { return FString::Printf(TEXT("Component %s Mobility is not Static"), *FastGeo::GetComponentShortName(InComponent)); } };
 	}
 
 	bool bIsComponentVisible = InComponent->IsVisible();
@@ -662,7 +669,7 @@ FFastGeoTransformResult UFastGeoWorldPartitionRuntimeCellTransformer::CanTransfo
 	{
 		if (!StaticMeshComponent->GetStaticMesh())
 		{
-			return { false, [InComponent] { return FString::Printf(TEXT("Component %s has an invalid static mesh"), *FastGeo::GetComponentShortName(InComponent)); } };
+			return { EFastGeoTransform::Discard, [InComponent] { return FString::Printf(TEXT("Component %s has an invalid static mesh"), *FastGeo::GetComponentShortName(InComponent)); } };
 		}
 
 		// Make sure BodyInstance CollisionEnabled is updated first before testing below
@@ -673,12 +680,12 @@ FFastGeoTransformResult UFastGeoWorldPartitionRuntimeCellTransformer::CanTransfo
 	{
 		if (!SkinnedMeshComponent->GetSkinnedAsset())
 		{
-			return { false, [InComponent] { return FString::Printf(TEXT("Component %s has an invalid skinned asset"), *FastGeo::GetComponentShortName(InComponent)); } };
+			return { EFastGeoTransform::Discard, [InComponent] { return FString::Printf(TEXT("Component %s has an invalid skinned asset"), *FastGeo::GetComponentShortName(InComponent)); } };
 		}
 
 		if (SkinnedMeshComponent->LeaderPoseComponent.IsValid())
 		{
-			return { false, [InComponent] { return FString::Printf(TEXT("Skinned mesh Component %s has a leader pose component"), *FastGeo::GetComponentShortName(InComponent)); } };
+			return { EFastGeoTransform::Reject, [InComponent] { return FString::Printf(TEXT("Skinned mesh Component %s has a leader pose component"), *FastGeo::GetComponentShortName(InComponent)); } };
 		}
 
 		bIsComponentVisible &= !SkinnedMeshComponent->bHideSkin;
@@ -689,105 +696,33 @@ FFastGeoTransformResult UFastGeoWorldPartitionRuntimeCellTransformer::CanTransfo
 	// If collision is enabled, only allow if async physics state creation and destruction are supported 
 	if (bIsCollisionEnabled && (!InComponent->AllowsAsyncPhysicsStateCreation() || !InComponent->AllowsAsyncPhysicsStateDestruction()))
 	{
-		return { false, [InComponent] { return FString::Printf(TEXT("Component %s has collision enabled but doesn't allow asynchronous physics state creation/destruction"), *FastGeo::GetComponentShortName(InComponent)); } };
+		return { EFastGeoTransform::Reject, [InComponent] { return FString::Printf(TEXT("Component %s has collision enabled but doesn't allow asynchronous physics state creation/destruction"), *FastGeo::GetComponentShortName(InComponent)); } };
 	}
 
 	// Disallow transform if collision is disabled and component is not visible
 	if (!bIsCollisionEnabled && !bIsComponentVisible)
 	{
-		return { false, [InComponent] { return FString::Printf(TEXT("Component %s has no collision and is not visible"), *FastGeo::GetComponentShortName(InComponent)); } };
+		return { EFastGeoTransform::Discard, [InComponent] { return FString::Printf(TEXT("Component %s has no collision and is not visible"), *FastGeo::GetComponentShortName(InComponent)); } };
 	}
 
-	return true;
-}
-
-bool UFastGeoWorldPartitionRuntimeCellTransformer::CanRemoveActor(AActor* InActor, TSet<UActorComponent*>& InIgnoredComponents) const
-{
-	auto IsComponentIgnored = [this, &InIgnoredComponents](UActorComponent* Component)
-	{
-		if (InIgnoredComponents.Contains(Component))
-		{
-			return true;
-		}
-
-		if (Component->IsEditorOnly())
-		{
-			return true;
-		}
-
-		if (UStaticMeshComponent* StaticMeshComponent = Cast<UStaticMeshComponent>(Component))
-		{
-			if (!StaticMeshComponent->GetStaticMesh())
-			{
-				return true;
-			}
-
-			if (!StaticMeshComponent->IsVisible() && !FastGeo::IsCollisionEnabled(StaticMeshComponent))
-			{
-				return true;
-			}
-		}
-
-		if (USkinnedMeshComponent* SkinnedMeshComponent = Cast<USkinnedMeshComponent>(Component))
-		{
-			if (!SkinnedMeshComponent->GetSkinnedAsset())
-			{
-				return true;
-			}
-
-			const bool bIsVisible = !SkinnedMeshComponent->bHideSkin && SkinnedMeshComponent->IsVisible();
-			if (!bIsVisible && !FastGeo::IsCollisionEnabled(SkinnedMeshComponent))
-			{
-				return true;
-			}
-		}
-
-		UClass* ComponentClass = Component->GetClass();
-		for (const TSubclassOf<UActorComponent>& IgnoredComponentClass : IgnoredRemainingComponentClasses)
-		{
-			if (ComponentClass->IsChildOf(IgnoredComponentClass))
-			{
-				return true;
-			}
-		}
-
-		for (const TSubclassOf<UActorComponent>& IgnoredComponentClass : IgnoredRemainingExactComponentClasses)
-		{
-			if (ComponentClass == IgnoredComponentClass)
-			{
-				return true;
-			}
-		}
-
-		return false;
-	};
-
-	uint32 NumRuntimeComponents = 0;
-	for (UActorComponent* Component : InActor->GetComponents())
+	UClass* ComponentClass = InComponent->GetClass();
+	for (const TSubclassOf<UActorComponent>& IgnoredComponentClass : IgnoredRemainingComponentClasses)
 	{
-		if (!IsComponentIgnored(Component))
+		if (ComponentClass->IsChildOf(IgnoredComponentClass))
 		{
-			NumRuntimeComponents++;
+			return { EFastGeoTransform::Discard };
 		}
 	}
 
-	if (!NumRuntimeComponents)
+	for (const TSubclassOf<UActorComponent>& IgnoredComponentClass : IgnoredRemainingExactComponentClasses)
 	{
-		return true;
-	}
-
-	if (NumRuntimeComponents == 1)
-	{
-		if (USceneComponent* RootComponent = InActor->GetRootComponent())
+		if (ComponentClass == IgnoredComponentClass)
 		{
-			if (ExactCast<USceneComponent>(RootComponent))
-			{
-				return true;
-			}
+			return { EFastGeoTransform::Discard };
 		}
 	}
 
-	return false;
+	return EFastGeoTransform::Allow;
 }
 
 void UFastGeoWorldPartitionRuntimeCellTransformer::PreEditChange(FProperty* InPropertyAboutToChange)
@@ -808,21 +743,21 @@ void UFastGeoWorldPartitionRuntimeCellTransformer::PostEditChangeProperty(FPrope
 
 bool FFastGeoTransformResult::ShouldReport = false;
 
-FFastGeoTransformResult::FFastGeoTransformResult(bool bInSuccess, const TCHAR* FailureReason)
-	: bSuccess(bInSuccess)
+FFastGeoTransformResult::FFastGeoTransformResult(EFastGeoTransform InTransformResult, const TCHAR* InFailureReason)
+	: TransformResult(InTransformResult)
 {
-	if (!bInSuccess && FailureReason && ShouldReport)
+	if (TransformResult != EFastGeoTransform::Allow && InFailureReason && ShouldReport)
 	{
-		UE_LOG(LogFastGeoStreaming, Log, TEXT("  * Can't transform: %s"), FailureReason);
+		UE_LOG(LogFastGeoStreaming, Log, TEXT("  * Can't transform: %s"), InFailureReason);
 	}
 }
 
-FFastGeoTransformResult::FFastGeoTransformResult(bool bInSuccess, TFunctionRef<FString ()> FailureReasonFunc)
-	: bSuccess(bInSuccess)
+FFastGeoTransformResult::FFastGeoTransformResult(EFastGeoTransform InTransformResult, TFunctionRef<FString()> InFailureReasonFunc)
+	: TransformResult(InTransformResult)
 {
-	if (!bInSuccess && ShouldReport)
+	if (TransformResult != EFastGeoTransform::Allow && ShouldReport)
 	{
-		UE_LOG(LogFastGeoStreaming, Log, TEXT("  * Can't transform: %s"), *FailureReasonFunc());
+		UE_LOG(LogFastGeoStreaming, Log, TEXT("  * Can't transform: %s"), *InFailureReasonFunc());
 	}
 }
 
diff --git a/Engine/Plugins/Experimental/FastGeoStreaming/Source/FastGeoStreaming/Private/FastGeoWorldPartitionRuntimeCellTransformer.h b/Engine/Plugins/Experimental/FastGeoStreaming/Source/FastGeoStreaming/Private/FastGeoWorldPartitionRuntimeCellTransformer.h
--- a/Engine/Plugins/Experimental/FastGeoStreaming/Source/FastGeoStreaming/Private/FastGeoWorldPartitionRuntimeCellTransformer.h
+++ b/Engine/Plugins/Experimental/FastGeoStreaming/Source/FastGeoStreaming/Private/FastGeoWorldPartitionRuntimeCellTransformer.h
@@ -11,19 +11,33 @@
 
 class UActorComponent;
 
+enum class EFastGeoTransform : uint32
+{
+	Allow,		//!< Actor or component can be transformed.
+	Reject,		//!< Actor or component can't be transformed.
+	Discard,	//!< Actor or component is not relevant and can be fully discarded without impact on the game.
+	MAX
+};
+
+inline constexpr uint32 EnumToIndex(EFastGeoTransform Value)
+{
+	return static_cast<uint32>(Value);
+}
+
 /** Used to conditionally report the error result when not a success. */
 struct FFastGeoTransformResult
 {
 public:
-	FFastGeoTransformResult(bool bInSuccess, const TCHAR* FailureReason = nullptr);
-	FFastGeoTransformResult(bool bInSuccess, TFunctionRef<FString()> FailureReasonFunc);
+	FFastGeoTransformResult(EFastGeoTransform InTransformResult, const TCHAR* FailureReason = nullptr);
+	FFastGeoTransformResult(EFastGeoTransform InTransformResult, TFunctionRef<FString()> FailureReasonFunc);
 
-	explicit operator bool() const { return bSuccess; }
+	EFastGeoTransform GetResult() const { return TransformResult; }
+	uint32 GetResultIndex() const { return EnumToIndex(TransformResult); }
 
 	static bool ShouldReport;
 
 private:
-	bool bSuccess;
+	EFastGeoTransform TransformResult;
 };
 
 UCLASS()
@@ -50,7 +64,7 @@ public:
 
 private:
 	bool IsBlueprintActorWithLogic(AActor* InActor) const;
-	FFastGeoTransformResult CanTransformActor(AActor* InActor, bool& bOutIsActorFullyTransformable, TSet<UActorComponent*>& OutTransformableComponents) const;
+	FFastGeoTransformResult CanTransformActor(AActor* InActor, bool& bOutIsActorFullyTransformable, TArray<UActorComponent*>& OutTransformableComponents) const;
 	FFastGeoTransformResult IsAllowedActorClass(AActor* InActor) const;
 	FFastGeoTransformResult CanTransformComponent(UPrimitiveComponent* InComponent) const;
 	FFastGeoTransformResult IsAllowedComponentClass(UActorComponent* InComponent) const;
