FastGeo - Discard ISMC / ISKMC that have no instances

diff --git a/Engine/Plugins/Experimental/FastGeoStreaming/Source/FastGeoStreaming/Private/FastGeoInstancedSkinnedMeshComponent.cpp b/Engine/Plugins/Experimental/FastGeoStreaming/Source/FastGeoStreaming/Private/FastGeoInstancedSkinnedMeshComponent.cpp
--- a/Engine/Plugins/Experimental/FastGeoStreaming/Source/FastGeoStreaming/Private/FastGeoInstancedSkinnedMeshComponent.cpp
+++ b/Engine/Plugins/Experimental/FastGeoStreaming/Source/FastGeoStreaming/Private/FastGeoInstancedSkinnedMeshComponent.cpp
@@ -60,6 +60,9 @@ void FFastGeoInstancedSkinnedMeshComponent::InitializeFromComponent(UActorCompon
 
 	LocalBounds = FInstancedSkinnedMeshComponentHelper::CalcBounds(*this, FTransform::Identity);
 	WorldBounds = FInstancedSkinnedMeshComponentHelper::CalcBounds(*this, WorldTransform);
+
+	// ISKMC with no instances should never be transformed to FastGeo
+	check(!InstanceData.IsEmpty());
 }
 	
 void FFastGeoInstancedSkinnedMeshComponent::ResetSceneProxyDescUnsupportedProperties()
diff --git a/Engine/Plugins/Experimental/FastGeoStreaming/Source/FastGeoStreaming/Private/FastGeoInstancedStaticMeshComponent.cpp b/Engine/Plugins/Experimental/FastGeoStreaming/Source/FastGeoStreaming/Private/FastGeoInstancedStaticMeshComponent.cpp
--- a/Engine/Plugins/Experimental/FastGeoStreaming/Source/FastGeoStreaming/Private/FastGeoInstancedStaticMeshComponent.cpp
+++ b/Engine/Plugins/Experimental/FastGeoStreaming/Source/FastGeoStreaming/Private/FastGeoInstancedStaticMeshComponent.cpp
@@ -64,6 +64,9 @@ void FFastGeoInstancedStaticMeshComponent::InitializeFromComponent(UActorCompone
 	PerInstanceSMCustomData = ISMComponent->PerInstanceSMCustomData;
 	InstancingRandomSeed = ISMComponent->InstancingRandomSeed;
 
+	// ISMC with no instances should never be transformed to FastGeo
+	check(!PerInstanceSMData.IsEmpty());
+
 	LocalBounds = CalculateBounds(EBoundsType::LocalBounds);
 	WorldBounds = CalculateBounds(EBoundsType::WorldBounds);
 	NavigationBounds = CalculateBounds(EBoundsType::NavigationBounds).GetBox();
diff --git a/Engine/Plugins/Experimental/FastGeoStreaming/Source/FastGeoStreaming/Private/FastGeoWorldPartitionRuntimeCellTransformer.cpp b/Engine/Plugins/Experimental/FastGeoStreaming/Source/FastGeoStreaming/Private/FastGeoWorldPartitionRuntimeCellTransformer.cpp
--- a/Engine/Plugins/Experimental/FastGeoStreaming/Source/FastGeoStreaming/Private/FastGeoWorldPartitionRuntimeCellTransformer.cpp
+++ b/Engine/Plugins/Experimental/FastGeoStreaming/Source/FastGeoStreaming/Private/FastGeoWorldPartitionRuntimeCellTransformer.cpp
@@ -676,6 +676,14 @@ FFastGeoTransformResult UFastGeoWorldPartitionRuntimeCellTransformer::CanTransfo
 		StaticMeshComponent->UpdateCollisionFromStaticMesh();
 	}
 
+	if (UInstancedStaticMeshComponent* InstancedStaticMeshComponent = Cast<UInstancedStaticMeshComponent>(InComponent))
+	{
+		if (InstancedStaticMeshComponent->GetNumInstances() == 0)
+		{
+			return { EFastGeoTransform::Discard, [InComponent] { return FString::Printf(TEXT("Component %s has no instances"), *FastGeo::GetComponentShortName(InComponent)); } };
+		}
+	}
+
 	if (USkinnedMeshComponent* SkinnedMeshComponent = Cast<USkinnedMeshComponent>(InComponent))
 	{
 		if (!SkinnedMeshComponent->GetSkinnedAsset())
@@ -691,6 +699,14 @@ FFastGeoTransformResult UFastGeoWorldPartitionRuntimeCellTransformer::CanTransfo
 		bIsComponentVisible &= !SkinnedMeshComponent->bHideSkin;
 	}
 
+	if (UInstancedSkinnedMeshComponent* InstancedSkinnedMeshComponent = Cast<UInstancedSkinnedMeshComponent>(InComponent))
+	{
+		if (InstancedSkinnedMeshComponent->GetInstanceCount() == 0)
+		{
+			return { EFastGeoTransform::Discard, [InComponent] { return FString::Printf(TEXT("Component %s has no instances"), *FastGeo::GetComponentShortName(InComponent)); } };
+		}
+	}
+
 	const bool bIsCollisionEnabled = FastGeo::IsCollisionEnabled(InComponent);
 
 	// If collision is enabled, only allow if async physics state creation and destruction are supported 
