FastGeo: Simplified fastegeo actor coloration: blue for fastgeo and red for non-fasteo.

diff --git a/Engine/Plugins/Experimental/FastGeoStreaming/Source/FastGeoStreaming/Private/FastGeoComponent.cpp b/Engine/Plugins/Experimental/FastGeoStreaming/Source/FastGeoStreaming/Private/FastGeoComponent.cpp
--- a/Engine/Plugins/Experimental/FastGeoStreaming/Source/FastGeoStreaming/Private/FastGeoComponent.cpp
+++ b/Engine/Plugins/Experimental/FastGeoStreaming/Source/FastGeoStreaming/Private/FastGeoComponent.cpp
@@ -66,7 +66,7 @@ bool FFastGeoComponent::IsRegistered() const
 
 FLinearColor FFastGeoComponent::GetDebugColor() const
 {
-	return UFastGeoWorldSubsystem::IsEnableDebugView() ? GetDebugColorImp() : FLinearColor::White;
+	return UFastGeoWorldSubsystem::IsEnableDebugView() ? FLinearColor::Blue : FLinearColor::White;
 }
 
 void FFastGeoComponent::Serialize(FArchive& Ar)
diff --git a/Engine/Plugins/Experimental/FastGeoStreaming/Source/FastGeoStreaming/Private/FastGeoComponent.h b/Engine/Plugins/Experimental/FastGeoStreaming/Source/FastGeoStreaming/Private/FastGeoComponent.h
--- a/Engine/Plugins/Experimental/FastGeoStreaming/Source/FastGeoStreaming/Private/FastGeoComponent.h
+++ b/Engine/Plugins/Experimental/FastGeoStreaming/Source/FastGeoStreaming/Private/FastGeoComponent.h
@@ -64,8 +64,6 @@ public:
 #endif
 
 protected:
-	virtual FLinearColor GetDebugColorImp() const { return FLinearColor::White; }
-
 	// Persistent Data
 	int32 ComponentIndex;
 
diff --git a/Engine/Plugins/Experimental/FastGeoStreaming/Source/FastGeoStreaming/Private/FastGeoInstancedSkinnedMeshComponent.cpp b/Engine/Plugins/Experimental/FastGeoStreaming/Source/FastGeoStreaming/Private/FastGeoInstancedSkinnedMeshComponent.cpp
--- a/Engine/Plugins/Experimental/FastGeoStreaming/Source/FastGeoStreaming/Private/FastGeoInstancedSkinnedMeshComponent.cpp
+++ b/Engine/Plugins/Experimental/FastGeoStreaming/Source/FastGeoStreaming/Private/FastGeoInstancedSkinnedMeshComponent.cpp
@@ -21,11 +21,6 @@ FFastGeoInstancedSkinnedMeshComponent::FFastGeoInstancedSkinnedMeshComponent(int
 	InstanceDataManager = MakeShared<FInstanceDataManager>(nullptr);
 }
 
-FLinearColor FFastGeoInstancedSkinnedMeshComponent::GetDebugColorImp() const
-{
-	return GetOwnerComponentCluster()->IsA<FFastGeoHLOD>() ? FColor::Turquoise : FColor::Yellow;
-}
-
 void FFastGeoInstancedSkinnedMeshComponent::Serialize(FArchive& Ar)
 {
 	Super::Serialize(Ar);
diff --git a/Engine/Plugins/Experimental/FastGeoStreaming/Source/FastGeoStreaming/Private/FastGeoInstancedSkinnedMeshComponent.h b/Engine/Plugins/Experimental/FastGeoStreaming/Source/FastGeoStreaming/Private/FastGeoInstancedSkinnedMeshComponent.h
--- a/Engine/Plugins/Experimental/FastGeoStreaming/Source/FastGeoStreaming/Private/FastGeoInstancedSkinnedMeshComponent.h
+++ b/Engine/Plugins/Experimental/FastGeoStreaming/Source/FastGeoStreaming/Private/FastGeoInstancedSkinnedMeshComponent.h
@@ -21,7 +21,6 @@ public:
 protected:
 	//~ Being FFastGeoComponent interface
 	virtual void Serialize(FArchive& Ar) override;
-	virtual FLinearColor GetDebugColorImp() const override;
 #if WITH_EDITOR
 	virtual void InitializeFromComponent(UActorComponent* Component) override;
 #endif
diff --git a/Engine/Plugins/Experimental/FastGeoStreaming/Source/FastGeoStreaming/Private/FastGeoInstancedStaticMeshComponent.cpp b/Engine/Plugins/Experimental/FastGeoStreaming/Source/FastGeoStreaming/Private/FastGeoInstancedStaticMeshComponent.cpp
--- a/Engine/Plugins/Experimental/FastGeoStreaming/Source/FastGeoStreaming/Private/FastGeoInstancedStaticMeshComponent.cpp
+++ b/Engine/Plugins/Experimental/FastGeoStreaming/Source/FastGeoStreaming/Private/FastGeoInstancedStaticMeshComponent.cpp
@@ -41,11 +41,6 @@ void FFastGeoInstancedStaticMeshComponent::Serialize(FArchive& Ar)
 	FArchive_Serialize_BitfieldBool(Ar, SceneProxyDesc.bUseGpuLodSelection);
 }
 
-FLinearColor FFastGeoInstancedStaticMeshComponent::GetDebugColorImp() const
-{
-	return GetOwnerComponentCluster()->IsA<FFastGeoHLOD>() ? FColor::Red : FColor::Orange;
-}
-
 #if WITH_EDITOR
 void FFastGeoInstancedStaticMeshComponent::InitializeSceneProxyDescFromComponent(UActorComponent* Component)
 {
diff --git a/Engine/Plugins/Experimental/FastGeoStreaming/Source/FastGeoStreaming/Private/FastGeoInstancedStaticMeshComponent.h b/Engine/Plugins/Experimental/FastGeoStreaming/Source/FastGeoStreaming/Private/FastGeoInstancedStaticMeshComponent.h
--- a/Engine/Plugins/Experimental/FastGeoStreaming/Source/FastGeoStreaming/Private/FastGeoInstancedStaticMeshComponent.h
+++ b/Engine/Plugins/Experimental/FastGeoStreaming/Source/FastGeoStreaming/Private/FastGeoInstancedStaticMeshComponent.h
@@ -37,10 +37,6 @@ public:
 	//~ End FFastGeoPrimitiveComponent interface
 
 protected:
-	//~ Being FFastGeoComponent interface
-	virtual FLinearColor GetDebugColorImp() const override;
-	//~ End FFastGeoComponent interface
-
 	//~ Begin FFastGeoPrimitiveComponent interface
 	virtual FPrimitiveSceneProxyDesc& GetSceneProxyDesc() override { return SceneProxyDesc; }
 	virtual const FPrimitiveSceneProxyDesc& GetSceneProxyDesc() const override { return SceneProxyDesc; }
diff --git a/Engine/Plugins/Experimental/FastGeoStreaming/Source/FastGeoStreaming/Private/FastGeoSkinnedMeshComponent.cpp b/Engine/Plugins/Experimental/FastGeoStreaming/Source/FastGeoStreaming/Private/FastGeoSkinnedMeshComponent.cpp
--- a/Engine/Plugins/Experimental/FastGeoStreaming/Source/FastGeoStreaming/Private/FastGeoSkinnedMeshComponent.cpp
+++ b/Engine/Plugins/Experimental/FastGeoStreaming/Source/FastGeoStreaming/Private/FastGeoSkinnedMeshComponent.cpp
@@ -103,11 +103,6 @@ void FFastGeoSkinnedMeshComponentBase::ResetSceneProxyDescUnsupportedProperties(
 }
 #endif
 
-FLinearColor FFastGeoSkinnedMeshComponentBase::GetDebugColorImp() const
-{
-	return GetOwnerComponentCluster()->IsA<FFastGeoHLOD>() ? FColor::Purple : FColor::Green;
-}
-
 void FFastGeoSkinnedMeshComponentBase::InitializeSceneProxyDescDynamicProperties()
 {
 	Super::InitializeSceneProxyDescDynamicProperties();
diff --git a/Engine/Plugins/Experimental/FastGeoStreaming/Source/FastGeoStreaming/Private/FastGeoSkinnedMeshComponent.h b/Engine/Plugins/Experimental/FastGeoStreaming/Source/FastGeoStreaming/Private/FastGeoSkinnedMeshComponent.h
--- a/Engine/Plugins/Experimental/FastGeoStreaming/Source/FastGeoStreaming/Private/FastGeoSkinnedMeshComponent.h
+++ b/Engine/Plugins/Experimental/FastGeoStreaming/Source/FastGeoStreaming/Private/FastGeoSkinnedMeshComponent.h
@@ -38,10 +38,6 @@ public:
 	//~ End FFastGeoPrimitiveComponent interface
 
 protected:
-	//~ Begin FFastGeoComponent interface
-	virtual FLinearColor GetDebugColorImp() const override;
-	//~ End FFastGeoComponent interface
-
 	//~ Begin FFastGeoPrimitiveComponent interface
 #if WITH_EDITOR
 	virtual void InitializeSceneProxyDescFromComponent(UActorComponent* Component) override;
diff --git a/Engine/Plugins/Experimental/FastGeoStreaming/Source/FastGeoStreaming/Private/FastGeoStaticMeshComponent.cpp b/Engine/Plugins/Experimental/FastGeoStreaming/Source/FastGeoStreaming/Private/FastGeoStaticMeshComponent.cpp
--- a/Engine/Plugins/Experimental/FastGeoStreaming/Source/FastGeoStreaming/Private/FastGeoStaticMeshComponent.cpp
+++ b/Engine/Plugins/Experimental/FastGeoStreaming/Source/FastGeoStreaming/Private/FastGeoStaticMeshComponent.cpp
@@ -64,11 +64,6 @@ const Nanite::FResources* FFastGeoStaticMeshComponentBase::GetNaniteResources()
 	return nullptr;
 }
 
-FLinearColor FFastGeoStaticMeshComponentBase::GetDebugColorImp() const
-{
-	return GetOwnerComponentCluster()->IsA<FFastGeoHLOD>() ? FColor::Blue : FColor::Cyan;
-}
-
 UBodySetup* FFastGeoStaticMeshComponentBase::GetBodySetup() const
 {
 	return GetStaticMesh() ? GetStaticMesh()->GetBodySetup() : nullptr;
diff --git a/Engine/Plugins/Experimental/FastGeoStreaming/Source/FastGeoStreaming/Private/FastGeoStaticMeshComponent.h b/Engine/Plugins/Experimental/FastGeoStreaming/Source/FastGeoStreaming/Private/FastGeoStaticMeshComponent.h
--- a/Engine/Plugins/Experimental/FastGeoStreaming/Source/FastGeoStreaming/Private/FastGeoStaticMeshComponent.h
+++ b/Engine/Plugins/Experimental/FastGeoStreaming/Source/FastGeoStreaming/Private/FastGeoStaticMeshComponent.h
@@ -75,10 +75,6 @@ public:
 	//~ End FFastGeoPrimitiveComponent interface
 
 protected:
-	//~ Begin FFastGeoComponent interface
-	virtual FLinearColor GetDebugColorImp() const override;
-	//~ End FFastGeoComponent interface
-
 	//~ Begin FFastGeoPrimitiveComponent interface
 #if WITH_EDITOR
 	virtual void InitializeSceneProxyDescFromComponent(UActorComponent* Component) override;
diff --git a/Engine/Plugins/Experimental/FastGeoStreaming/Source/FastGeoStreaming/Private/FastGeoWorldSubsystem.cpp b/Engine/Plugins/Experimental/FastGeoStreaming/Source/FastGeoStreaming/Private/FastGeoWorldSubsystem.cpp
--- a/Engine/Plugins/Experimental/FastGeoStreaming/Source/FastGeoStreaming/Private/FastGeoWorldSubsystem.cpp
+++ b/Engine/Plugins/Experimental/FastGeoStreaming/Source/FastGeoStreaming/Private/FastGeoWorldSubsystem.cpp
@@ -66,7 +66,7 @@ UFastGeoWorldSubsystem::UFastGeoWorldSubsystem()
 		FastGeoColorHandler.bAvailalbleInEditor = false;
 		FastGeoColorHandler.GetColorFunc = [](const UPrimitiveComponent* InPrimitiveComponent)
 		{
-			return FLinearColor::White;
+			return FLinearColor::Red;
 		};
 		FastGeoColorHandler.ActivateFunc = [&UpdatePrimitivesColor]()
 		{
