FastGeo - Added CVars to allow throttling of the render state creation tasks
* FastGeo.AsyncRenderStateTask.ParallelWorkerCount - Limit the number of threads. Default is 1.
* FastGeo.AsyncRenderStateTask.TimeBudgetMS - Limit the processing time of the async task. Default is no limit.
* FastGeo.AsyncRenderStateTask.MaxNumComponentsToProcess - Limit the number of components procesed per frame. Default is no limit.
* Only async render state creation is throttled at the moment

This can help prevent processing spikes on the render thread as otherwise it would need to add all the primitives in a single frame.

diff --git a/Engine/Plugins/Experimental/FastGeoStreaming/Source/FastGeoStreaming/Private/FastGeoContainer.cpp b/Engine/Plugins/Experimental/FastGeoStreaming/Source/FastGeoStreaming/Private/FastGeoContainer.cpp
--- a/Engine/Plugins/Experimental/FastGeoStreaming/Source/FastGeoStreaming/Private/FastGeoContainer.cpp
+++ b/Engine/Plugins/Experimental/FastGeoStreaming/Source/FastGeoStreaming/Private/FastGeoContainer.cpp
@@ -6,19 +6,46 @@
 #include "FastGeoComponent.h"
 #include "FastGeoStreamingManager.h"
 #include "FastGeoLog.h"
+#include "AI/Navigation/NavigationElement.h"
 #include "Async/ParallelFor.h"
 #include "Engine/Level.h"
 #include "Engine/SkeletalMesh.h"
 #include "Engine/World.h"
+#include "Misc/Timeout.h"
+#include "NavigationSystem.h"
 #include "PhysicsEngine/BodySetup.h"
 #include "PhysicsEngine/PhysicsAsset.h"
 #include "PhysicalMaterials/PhysicalMaterial.h"
 #include "Physics/Experimental/PhysScene_Chaos.h"
+#include "ProfilingDebugging/CountersTrace.h"
+#include "PSOPrecache.h"
 #include "WorldPartition/WorldPartitionHelpers.h"
 #include "WorldPartition/WorldPartition.h"
-#include "NavigationSystem.h"
-#include "AI/Navigation/NavigationElement.h"
-#include "PSOPrecache.h"
+
+namespace FastGeo
+{
+	static int32 GAsyncRenderStateTaskParallelWorkerCount = 1;
+	static FAutoConsoleVariableRef CVarAsyncRenderStateTaskParallelWorkerCount(
+		TEXT("FastGeo.AsyncRenderStateTask.ParallelWorkerCount"),
+		GAsyncRenderStateTaskParallelWorkerCount,
+		TEXT("Set the max number of workers to use when creating FastGeo render state. ")
+		TEXT("Only taken into account if value is greater than 1."),
+		ECVF_Default);
+
+	static float GAsyncRenderStateTaskTimeBudgetMS = 0.0f;
+	static FAutoConsoleVariableRef CVarAsyncRenderStateTaskTimeBudgetMS(
+		TEXT("FastGeo.AsyncRenderStateTask.TimeBudgetMS"),
+		GAsyncRenderStateTaskTimeBudgetMS,
+		TEXT("Maximum time budget in milliseconds for the async render state tasks (0 = no time limit)"),
+		ECVF_Default);
+
+	static int32 GAsyncRenderStateTaskMaxNumComponentsToProcess = 0;
+	static FAutoConsoleVariableRef CVarAsyncRenderStateTaskMaxNumComponentsToProcess(
+		TEXT("FastGeo.AsyncRenderStateTask.MaxNumComponentsToProcess"),
+		GAsyncRenderStateTaskMaxNumComponentsToProcess,
+		TEXT("Maximum number of components to process (0 = no component limit)"),
+		ECVF_Default);
+}
 
 class FFastGeoAssetRemapArchive : public FArchiveProxy
 {
@@ -85,6 +112,8 @@ void UFastGeoContainer::Register()
 
 	if (!bIsRegistered)
 	{
+		PendingCreate.Reset();
+
 		check(GetWorld()->IsGameWorld());
 #if WITH_EDITOR
 		// In PIE we need to initialize dynamic properties as there's no serialization
@@ -93,20 +122,29 @@ void UFastGeoContainer::Register()
 		UWorldPartition* WorldPartition = FWorldPartitionHelpers::GetWorldPartition(this);
 		const bool bApplyWorldTransform = WorldPartition->HasInstanceTransform();
 		const FTransform& Transform = WorldPartition->GetInstanceTransform();
-		ForEachComponentCluster([bApplyWorldTransform , &Transform](FFastGeoComponentCluster& ComponentCluster)
+		ForEachComponentCluster([bApplyWorldTransform , &Transform, this](FFastGeoComponentCluster& ComponentCluster)
 		{
 			ComponentCluster.OnRegister();
 
-			if (bApplyWorldTransform)
+			ComponentCluster.ForEachComponent<FFastGeoPrimitiveComponent>([bApplyWorldTransform, &Transform, this](FFastGeoPrimitiveComponent& Component)
 			{
-				ComponentCluster.ForEachComponent<FFastGeoPrimitiveComponent>([&Transform](FFastGeoPrimitiveComponent& Component)
+				if (bApplyWorldTransform)
 				{
 					Component.ApplyWorldTransform(Transform);
-				});
-			}
+				}
+
+				if (Component.ShouldCreateRenderState())
+				{
+					PendingCreate.RenderState.ComponentsToProcess.Add(&Component);
+				}
+			});
 		});
 
-		FFastGeoStreamingManager::PushAsyncCreateRenderStateJob(this);
+		if (!PendingCreate.RenderState.ComponentsToProcess.IsEmpty())
+		{
+			FFastGeoStreamingManager::PushAsyncCreateRenderStateJob(this);
+		}
+
 		FFastGeoStreamingManager::PushAsyncCreatePhysicsStateJobs(this);
 		RegisterToNavigationSystem();
 
@@ -137,13 +175,25 @@ void UFastGeoContainer::Unregister()
 
 	if (bIsRegistered)
 	{
-		ForEachComponentCluster([](FFastGeoComponentCluster& ComponentCluster)
+		PendingDestroy.Reset();
+
+		ForEachComponentCluster([this](FFastGeoComponentCluster& ComponentCluster)
 		{
 			ComponentCluster.OnUnregister();
+
+			ComponentCluster.ForEachComponent<FFastGeoPrimitiveComponent>([this](FFastGeoPrimitiveComponent& Component)
+			{
+				PendingDestroy.RenderState.ComponentsToProcess.Add(&Component);
+			});
 		});
 
 		UnregisterFromNavigationSystem();
-		FFastGeoStreamingManager::PushAsyncDestroyRenderStateJob(this);
+		
+		if (!PendingDestroy.RenderState.ComponentsToProcess.IsEmpty())
+		{
+			FFastGeoStreamingManager::PushAsyncDestroyRenderStateJob(this);
+		}
+
 		FFastGeoStreamingManager::PushAsyncDestroyPhysicsStateJobs(this);
 		bIsRegistered = false;
 	}
@@ -156,13 +206,16 @@ void UFastGeoContainer::Unregister()
 
 void UFastGeoContainer::Tick(bool bWaitForCompletion)
 {
-	FFastGeoStreamingManager::ProcessAsyncRenderStateJobs(bWaitForCompletion);
-	UWorld* World = GetWorld();
-	check(World);
-	if (FPhysScene* PhysScene = World->GetPhysicsScene())
+	do
 	{
-		PhysScene->ProcessAsyncPhysicsStateJobs();
-	}
+		FFastGeoStreamingManager::ProcessAsyncRenderStateJobs(bWaitForCompletion);
+		UWorld* World = GetWorld();
+		check(World);
+		if (FPhysScene* PhysScene = World->GetPhysicsScene())
+		{
+			PhysScene->ProcessAsyncPhysicsStateJobs();
+		}
+	} while (bWaitForCompletion && HasAnyPendingTasks());
 }
 
 bool UFastGeoContainer::HasAnyPendingTasks() const
@@ -263,70 +316,83 @@ void UFastGeoContainer::UnregisterFromNavigationSystem()
 
 void UFastGeoContainer::OnCreateRenderStateBegin_GameThread()
 {
-	check(!PendingCreate.bRenderState);
-	PendingCreate.bRenderState = true;
+	check(PendingCreate.RenderState.NumToProcess == 0);
+	check(PendingCreate.RenderState.TotalNumProcessed != PendingCreate.RenderState.ComponentsToProcess.Num());
+
+	PendingCreate.RenderState.NumToProcess = PendingCreate.RenderState.ComponentsToProcess.Num() - PendingCreate.RenderState.TotalNumProcessed;
+	if (FastGeo::GAsyncRenderStateTaskMaxNumComponentsToProcess > 0)
+	{
+		PendingCreate.RenderState.NumToProcess = FMath::Min(FastGeo::GAsyncRenderStateTaskMaxNumComponentsToProcess, PendingCreate.RenderState.NumToProcess);
+	}
 }
 
 void UFastGeoContainer::OnDestroyRenderStateBegin_GameThread()
 {
-	check(!PendingDestroy.bRenderState);
-	PendingDestroy.bRenderState = true;
-}
+	check(PendingDestroy.RenderState.NumToProcess == 0);
+	check(PendingDestroy.RenderState.TotalNumProcessed != PendingDestroy.RenderState.ComponentsToProcess.Num());
 
-namespace FastGeo
-{
-	static int32 GCreateRenderStateWorkerCount = 0;
-	static FAutoConsoleVariableRef CVarCreateRenderStateWorkerCount(
-		TEXT("FastGeo.CreateRenderStateWorkerCount"),
-		GCreateRenderStateWorkerCount,
-		TEXT("Set the max number of workers to use when creating FastGeo render state. ")
-		TEXT("Only taken into account if value is greater than 1."),
-		ECVF_Default);
+	// TODO_FASTGEO: Evaluate if we need throttling of the async destruction task
+	PendingDestroy.RenderState.NumToProcess = PendingDestroy.RenderState.ComponentsToProcess.Num();
 }
 
 void UFastGeoContainer::OnCreateRenderState_Concurrent()
 {
-	TRACE_CPUPROFILER_EVENT_SCOPE(UFastGeoContainer::OnCreateRenderState_Concurrent);
+	static const int32 MinNumElementsToProcessPerThread = 8;
+	const int32 NumComponentsToProcess = PendingCreate.RenderState.NumToProcess;
+	const int32 NumThreads = FMath::Clamp(NumComponentsToProcess / MinNumElementsToProcessPerThread, 1, FastGeo::GAsyncRenderStateTaskParallelWorkerCount);
+	const bool bIsParallelForAllowed = NumThreads > 1 && FApp::ShouldUseThreadingForPerformance();
 
-	TArray<FFastGeoPrimitiveComponent*, TInlineAllocator<256>> Components;
-	ForEachComponentCluster([&Components](FFastGeoComponentCluster& ComponentCluster)
+	const bool bTimeLimit = FastGeo::GAsyncRenderStateTaskTimeBudgetMS != 0;
+	const double TimeBudgetSeconds = FastGeo::GAsyncRenderStateTaskTimeBudgetMS / 1000.0;
+	UE::FTimeout Timeout = bTimeLimit ? UE::FTimeout(TimeBudgetSeconds) : UE::FTimeout::Never();
+	
+	TAtomic<int32> NextIndex = 0;
+	TAtomic<int32> NumProcessed = 0;
+
+	ParallelFor(NumThreads, [&](int32 Index)
 	{
-		ComponentCluster.ForEachComponent<FFastGeoPrimitiveComponent>([&Components](FFastGeoPrimitiveComponent& Component)
+		TRACE_CPUPROFILER_EVENT_SCOPE("UFastGeoContainer::OnCreateRenderState_Concurrent");
+		FTaskTagScope Scope(ETaskTag::EParallelGameThread);
+
+		while(true)
 		{
-			if (Component.ShouldCreateRenderState())
+			// Time budget exceeded ?
+			if (Timeout.IsExpired())
 			{
-				Components.Add(&Component);
+				return; 
 			}
-		});
-	});
 
-	static int32 MinNumElementsToProcess = 8;
-	const int32 NumComponents = Components.Num();
-	const bool bIsParallelForAllowed = FApp::ShouldUseThreadingForPerformance() && (NumComponents >= MinNumElementsToProcess);
-	const int32 MinBatchSize = (bIsParallelForAllowed && (FastGeo::GCreateRenderStateWorkerCount > 0) && (NumComponents > FastGeo::GCreateRenderStateWorkerCount)) ? FMath::DivideAndRoundUp(NumComponents, FastGeo::GCreateRenderStateWorkerCount) : 1;
-	
-	ParallelFor(TEXT("UFastGeoContainer::OnCreateRenderState_Concurrent"), NumComponents, MinBatchSize, [&](int32 Index)
-	{
-		// TODO_FASTGEO : Can we call GetUsedMaterials from here? FInstancedStaticMeshSceneProxy -> FStaticMeshSceneProxy -> FStaticMeshSceneProxy::FLODInfo -> FMaterial::MaterialModifiesMeshPosition_RenderThread()
-		FTaskTagScope Scope(ETaskTag::EParallelGameThread);
-		Components[Index]->CreateRenderState(/*Context=*/nullptr);
+			// All work completed ?
+			int32 LocalIdx = NextIndex++;
+			if (LocalIdx >= NumComponentsToProcess)
+			{
+				return; 
+			}
+
+			// TODO_FASTGEO : Can we call GetUsedMaterials from here? FInstancedStaticMeshSceneProxy -> FStaticMeshSceneProxy -> FStaticMeshSceneProxy::FLODInfo -> FMaterial::MaterialModifiesMeshPosition_RenderThread()
+			FFastGeoPrimitiveComponent* ComponentToProcess = PendingCreate.RenderState.ComponentsToProcess[PendingCreate.RenderState.TotalNumProcessed + LocalIdx];
+			ComponentToProcess->CreateRenderState(/*Context = */ nullptr);
+
+			NumProcessed++;
+		}
 	}, bIsParallelForAllowed ? EParallelForFlags::None : EParallelForFlags::ForceSingleThread);
+
+	// Adjust NumRenderStatesToProcess to the actual number of states we processed. 
+	PendingCreate.RenderState.NumProcessed = NumProcessed.Load(EMemoryOrder::Relaxed);
 }
 
 void UFastGeoContainer::OnDestroyRenderState_Concurrent()
 {
 	TRACE_CPUPROFILER_EVENT_SCOPE(UFastGeoContainer::OnDestroyRenderState_Concurrent);
 
+	// TODO_FASTGEO: Evaluate if we need throttling of the async destruction task
 	{
 		FFastGeoPrimitiveComponent::FFastGeoDestroyRenderStateContext Context(GetWorld()->Scene);
 
-		ForEachComponentCluster([&Context](FFastGeoComponentCluster& ComponentCluster)
+		for (FFastGeoPrimitiveComponent* Component : PendingDestroy.RenderState.ComponentsToProcess)
 		{
-			ComponentCluster.ForEachComponent<FFastGeoPrimitiveComponent>([&Context](FFastGeoPrimitiveComponent& Component)
-			{
-				Component.DestroyRenderState(&Context);
-			});
-		});
+			Component->DestroyRenderState(&Context);
+		}
 
 #if DO_CHECK
 		bDestroyRenderCommandCompleted = !Context.HasPendingWork();
@@ -347,14 +413,33 @@ void UFastGeoContainer::OnDestroyRenderState_Concurrent()
 
 void UFastGeoContainer::OnCreateRenderStateEnd_GameThread()
 {
-	check(PendingCreate.bRenderState);
-	PendingCreate.bRenderState = false;
+	check(PendingCreate.RenderState.NumToProcess != 0);
+	check(PendingCreate.RenderState.TotalNumProcessed != PendingCreate.RenderState.ComponentsToProcess.Num());
+
+	PendingCreate.RenderState.TotalNumProcessed += PendingCreate.RenderState.NumProcessed;
+	PendingCreate.RenderState.NumToProcess = 0;
+	PendingCreate.RenderState.NumProcessed = 0;
+	
+
+	if (PendingCreate.RenderState.TotalNumProcessed < PendingCreate.RenderState.ComponentsToProcess.Num())
+	{
+		FFastGeoStreamingManager::PushAsyncCreateRenderStateJob(this);
+	}
 }
 
 void UFastGeoContainer::OnDestroyRenderStateEnd_GameThread()
 {
-	check(PendingDestroy.bRenderState);
-	PendingDestroy.bRenderState = false;
+	check(PendingDestroy.RenderState.NumToProcess != 0);
+	check(PendingDestroy.RenderState.TotalNumProcessed != PendingDestroy.RenderState.ComponentsToProcess.Num());
+
+	PendingDestroy.RenderState.TotalNumProcessed += PendingDestroy.RenderState.NumProcessed;
+	PendingDestroy.RenderState.NumToProcess = 0;
+	PendingDestroy.RenderState.NumProcessed = 0;	
+
+	if (PendingDestroy.RenderState.TotalNumProcessed < PendingDestroy.RenderState.ComponentsToProcess.Num())
+	{
+		FFastGeoStreamingManager::PushAsyncDestroyRenderStateJob(this);
+	}
 }
 
 void UFastGeoContainer::FinishDestroy()
diff --git a/Engine/Plugins/Experimental/FastGeoStreaming/Source/FastGeoStreaming/Private/FastGeoContainer.h b/Engine/Plugins/Experimental/FastGeoStreaming/Source/FastGeoStreaming/Private/FastGeoContainer.h
--- a/Engine/Plugins/Experimental/FastGeoStreaming/Source/FastGeoStreaming/Private/FastGeoContainer.h
+++ b/Engine/Plugins/Experimental/FastGeoStreaming/Source/FastGeoStreaming/Private/FastGeoContainer.h
@@ -25,13 +25,33 @@ class UMaterialInterface;
 
 struct FFastGeoComponentClusterPendingState
 {
-	FFastGeoComponentClusterPendingState()
-		: bRenderState(false)
-	{}
+	struct FRenderState
+	{
+		TArray<FFastGeoPrimitiveComponent*> ComponentsToProcess;
+
+		int32 NumToProcess;
+		int32 NumProcessed;
+		int32 TotalNumProcessed;
+	} RenderState;
 
-	bool bRenderState;
 	bool bPhysicsState;
-	bool HasAnyPendingState() const { return bRenderState || bPhysicsState; }
+	
+	FFastGeoComponentClusterPendingState()
+	{
+		Reset();
+	}
+
+	void Reset()
+	{
+		RenderState.ComponentsToProcess.Reset();
+		RenderState.NumToProcess = 0;
+		RenderState.NumProcessed = 0;
+		RenderState.TotalNumProcessed = 0;
+
+		bPhysicsState = false;
+	}
+
+	bool HasAnyPendingState() const { return RenderState.TotalNumProcessed < RenderState.ComponentsToProcess.Num() || bPhysicsState; }
 };
 
 UCLASS(Within = Level)
