FastGeo - Fixed issue introduced with throttling changes causing levels to never be removed from the world, and blocking unloads to never complete

diff --git a/Engine/Plugins/Experimental/FastGeoStreaming/Source/FastGeoStreaming/Private/FastGeoContainer.cpp b/Engine/Plugins/Experimental/FastGeoStreaming/Source/FastGeoStreaming/Private/FastGeoContainer.cpp
--- a/Engine/Plugins/Experimental/FastGeoStreaming/Source/FastGeoStreaming/Private/FastGeoContainer.cpp
+++ b/Engine/Plugins/Experimental/FastGeoStreaming/Source/FastGeoStreaming/Private/FastGeoContainer.cpp
@@ -386,6 +386,7 @@ void UFastGeoContainer::OnDestroyRenderState_Concurrent()
 	TRACE_CPUPROFILER_EVENT_SCOPE(UFastGeoContainer::OnDestroyRenderState_Concurrent);
 
 	// TODO_FASTGEO: Evaluate if we need throttling of the async destruction task
+	check(PendingDestroy.RenderState.NumToProcess == PendingDestroy.RenderState.ComponentsToProcess.Num());
 	{
 		FFastGeoPrimitiveComponent::FFastGeoDestroyRenderStateContext Context(GetWorld()->Scene);
 
@@ -394,6 +395,8 @@ void UFastGeoContainer::OnDestroyRenderState_Concurrent()
 			Component->DestroyRenderState(&Context);
 		}
 
+		PendingDestroy.RenderState.NumProcessed = PendingDestroy.RenderState.NumToProcess;
+
 #if DO_CHECK
 		bDestroyRenderCommandCompleted = !Context.HasPendingWork();
 #endif
