Added GetCollisionResponseToChannel , GetPhysicsBodyInstandeOwnerFromHitResult & GetPhysicsBodyInstandeOwnerFromOverlapResult in IPhysicsBodyInstanceOwner.
Adapted FastGeo to bind a custom FChaosUserDefinedEntity (FFastGeoPhysicsBodyInstanceOwner) to its physics objects so that hit result and overlap results can find an associated IPhysicsBodyInstanceOwner.
This allows for code that relied on component attached to a hit/overalp result to now also have access to a IPhysicsBodyInstanceOwner when available.

diff --git a/Engine/Plugins/Experimental/FastGeoStreaming/Source/FastGeoStreaming/FastGeoStreaming.Build.cs b/Engine/Plugins/Experimental/FastGeoStreaming/Source/FastGeoStreaming/FastGeoStreaming.Build.cs
--- a/Engine/Plugins/Experimental/FastGeoStreaming/Source/FastGeoStreaming/FastGeoStreaming.Build.cs
+++ b/Engine/Plugins/Experimental/FastGeoStreaming/Source/FastGeoStreaming/FastGeoStreaming.Build.cs
@@ -13,6 +13,7 @@ namespace UnrealBuildTool.Rules
 					"CoreUObject",
 					"Engine",
 					"NavigationSystem",
+					"Chaos",
 					"PhysicsCore",
 					"RenderCore",
 					"RHI",
diff --git a/Engine/Plugins/Experimental/FastGeoStreaming/Source/FastGeoStreaming/Private/FastGeoContainer.cpp b/Engine/Plugins/Experimental/FastGeoStreaming/Source/FastGeoStreaming/Private/FastGeoContainer.cpp
--- a/Engine/Plugins/Experimental/FastGeoStreaming/Source/FastGeoStreaming/Private/FastGeoContainer.cpp
+++ b/Engine/Plugins/Experimental/FastGeoStreaming/Source/FastGeoStreaming/Private/FastGeoContainer.cpp
@@ -15,6 +15,7 @@
 #include "NavigationSystem.h"
 #include "PhysicsEngine/BodySetup.h"
 #include "PhysicsEngine/PhysicsAsset.h"
+#include "PhysicsEngine/PhysicsObjectExternalInterface.h"
 #include "PhysicalMaterials/PhysicalMaterial.h"
 #include "Physics/Experimental/PhysScene_Chaos.h"
 #include "ProfilingDebugging/CountersTrace.h"
@@ -767,3 +768,14 @@ void UFastGeoContainer::CollectBodySetupsWithPhysicsMeshesToCreate(TSet<UBodySet
 		}
 	}
 }
+
+IPhysicsBodyInstanceOwner* UFastGeoContainer::ResolvePhysicsBodyInstanceOwner(Chaos::FConstPhysicsObjectHandle PhysicsObject)
+{
+	if (PhysicsObject)
+	{
+		FLockedReadPhysicsObjectExternalInterface PhysicsObjectInterface = FPhysicsObjectExternalInterface::LockRead(PhysicsObject);
+		FChaosUserDefinedEntity* UserDefinedEntity = PhysicsObjectInterface->GetUserDefinedEntity(PhysicsObject);
+		return FFastGeoPhysicsBodyInstanceOwner::GetPhysicsBodyInstanceOwner(UserDefinedEntity);
+	}
+	return nullptr;
+}
diff --git a/Engine/Plugins/Experimental/FastGeoStreaming/Source/FastGeoStreaming/Private/FastGeoContainer.h b/Engine/Plugins/Experimental/FastGeoStreaming/Source/FastGeoStreaming/Private/FastGeoContainer.h
--- a/Engine/Plugins/Experimental/FastGeoStreaming/Source/FastGeoStreaming/Private/FastGeoContainer.h
+++ b/Engine/Plugins/Experimental/FastGeoStreaming/Source/FastGeoStreaming/Private/FastGeoContainer.h
@@ -4,6 +4,7 @@
 
 #include "Engine/AssetUserData.h"
 #include "Physics/Experimental/AsyncPhysicsStateProcessorInterface.h"
+#include "PhysicsEngine/PhysicsBodyInstanceOwnerInterface.h"
 #include "PhysicsEngine/BodyInstance.h"
 #include "AI/Navigation/NavigationElement.h"
 #include "FastGeoComponentCluster.h"
@@ -18,10 +19,11 @@
 
 class UWorld;
 class ULevel;
-class FFastGeoComponent;
 class URuntimeVirtualTexture;
 class UStaticMesh;
 class UMaterialInterface;
+class FFastGeoComponent;
+class FChaosUserDefinedEntity;
 
 struct FFastGeoComponentClusterPendingState
 {
@@ -55,7 +57,7 @@ struct FFastGeoComponentClusterPendingState
 };
 
 UCLASS(Within = Level)
-class FASTGEOSTREAMING_API UFastGeoContainer : public UAssetUserData, public IAsyncPhysicsStateProcessor
+class FASTGEOSTREAMING_API UFastGeoContainer : public UAssetUserData, public IAsyncPhysicsStateProcessor, public IPhysicsBodyInstanceOwnerResolver
 {
 	GENERATED_BODY()
 
@@ -80,6 +82,10 @@ public:
 #endif
 	//~ End UObject Interface
 
+	//~Begin IPhysicsBodyInstanceOwnerResolver
+	virtual IPhysicsBodyInstanceOwner* ResolvePhysicsBodyInstanceOwner(Chaos::FConstPhysicsObjectHandle PhysicsObject) override;
+	//~End IPhysicsBodyInstanceOwnerResolver
+
 #if WITH_EDITOR
 	void OnCreated();
 	void AddComponentCluster(FFastGeoComponentCluster* ComponentCluster);
diff --git a/Engine/Plugins/Experimental/FastGeoStreaming/Source/FastGeoStreaming/Private/FastGeoInstancedStaticMeshComponent.cpp b/Engine/Plugins/Experimental/FastGeoStreaming/Source/FastGeoStreaming/Private/FastGeoInstancedStaticMeshComponent.cpp
--- a/Engine/Plugins/Experimental/FastGeoStreaming/Source/FastGeoStreaming/Private/FastGeoInstancedStaticMeshComponent.cpp
+++ b/Engine/Plugins/Experimental/FastGeoStreaming/Source/FastGeoStreaming/Private/FastGeoInstancedStaticMeshComponent.cpp
@@ -14,6 +14,8 @@
 #include "NaniteSceneProxy.h"
 #include "PhysicsEngine/BodySetup.h"
 #include "SceneInterface.h"
+#include "PhysicsEngine/PhysicsObjectExternalInterface.h"
+#include "PhysicsProxy/SingleParticlePhysicsProxy.h"
 
 const FFastGeoElementType FFastGeoInstancedStaticMeshComponent::Type(&FFastGeoStaticMeshComponentBase::Type);
 
@@ -242,6 +244,13 @@ void FFastGeoInstancedStaticMeshComponent::OnAsyncDestroyPhysicsStateBegin_GameT
 
 	FFastGeoComponent::OnAsyncDestroyPhysicsStateBegin_GameThread();
 
+	// Remove all user defined entities here
+	TArray<Chaos::FPhysicsObject*> PhysicsObjects = GetAllPhysicsObjects();
+	FPhysicsObjectExternalInterface::LockWrite(PhysicsObjects)->SetUserDefinedEntity(PhysicsObjects, nullptr);
+
+	// Reset BodyInstanceOwner
+	BodyInstanceOwner.Uninitialize();
+
 	// Move InstanceBodies in AsyncDestroyPhysicsStatePayload
 	check(AsyncDestroyPhysicsStatePayload.IsEmpty());
 	AsyncDestroyPhysicsStatePayload = MoveTemp(InstanceBodies);
@@ -265,6 +274,26 @@ void FFastGeoInstancedStaticMeshComponent::OnAsyncDestroyPhysicsState()
 	AsyncDestroyPhysicsStatePayload.Empty();
 }
 
+TArray<Chaos::FPhysicsObject*> FFastGeoInstancedStaticMeshComponent::GetAllPhysicsObjects() const
+{
+	auto GetPhysicsObjectById = [this](Chaos::FPhysicsObjectId Id) -> Chaos::FPhysicsObject*
+	{
+		if (!InstanceBodies.IsValidIndex(Id) || !InstanceBodies[Id] || !InstanceBodies[Id]->GetPhysicsActor())
+		{
+			return nullptr;
+		}
+		return InstanceBodies[Id]->GetPhysicsActor()->GetPhysicsObject();
+	};
+
+	TArray<Chaos::FPhysicsObject*> Objects;
+	Objects.Reserve(InstanceBodies.Num());
+	for (int32 Index = 0; Index < InstanceBodies.Num(); ++Index)
+	{
+		Objects.Add(GetPhysicsObjectById(Index));
+	}
+	return Objects;
+}
+
 void FFastGeoInstancedStaticMeshComponent::CreateAllInstanceBodies()
 {
 	TRACE_CPUPROFILER_EVENT_SCOPE(FFastGeoInstancedStaticMeshComponent::CreateAllInstanceBodies);
@@ -313,7 +342,16 @@ void FFastGeoInstancedStaticMeshComponent::CreateAllInstanceBodies()
 
 		if (InstanceBodiesSanitized.Num() > 0)
 		{
-			FBodyInstance::InitStaticBodies(MoveTemp(InstanceBodiesSanitized), MoveTemp(Transforms), BodySetup, nullptr, GetWorld()->GetPhysicsScene(), this);
+			// Initialize BodyInstanceOwner
+			BodyInstanceOwner.Initialize(this);
+
+			// Initialize body instances
+			FBodyInstance::InitStaticBodies(MoveTemp(InstanceBodiesSanitized), MoveTemp(Transforms), BodySetup, nullptr, GetWorld()->GetPhysicsScene(), &BodyInstanceOwner);
+
+			// Assign BodyInstanceOwner
+			TArray<Chaos::FPhysicsObject*> PhysicsObjects = GetAllPhysicsObjects();
+			FLockedWritePhysicsObjectExternalInterface Interface = FPhysicsObjectExternalInterface::LockWrite(PhysicsObjects);
+			Interface->SetUserDefinedEntity(PhysicsObjects, &BodyInstanceOwner);
 		}
 	}
 	else
diff --git a/Engine/Plugins/Experimental/FastGeoStreaming/Source/FastGeoStreaming/Private/FastGeoInstancedStaticMeshComponent.h b/Engine/Plugins/Experimental/FastGeoStreaming/Source/FastGeoStreaming/Private/FastGeoInstancedStaticMeshComponent.h
--- a/Engine/Plugins/Experimental/FastGeoStreaming/Source/FastGeoStreaming/Private/FastGeoInstancedStaticMeshComponent.h
+++ b/Engine/Plugins/Experimental/FastGeoStreaming/Source/FastGeoStreaming/Private/FastGeoInstancedStaticMeshComponent.h
@@ -34,6 +34,7 @@ public:
 	virtual bool IsNavigationRelevant() const override;
 	virtual bool DoCustomNavigableGeometryExport(FNavigableGeometryExport& GeomExport) const override;
 	virtual bool ShouldSkipNavigationDirtyAreaOnAddOrRemove() override { return true; }
+	virtual TArray<Chaos::FPhysicsObject*> GetAllPhysicsObjects() const override;
 	//~ End FFastGeoPrimitiveComponent interface
 
 protected:
diff --git a/Engine/Plugins/Experimental/FastGeoStreaming/Source/FastGeoStreaming/Private/FastGeoPrimitiveComponent.cpp b/Engine/Plugins/Experimental/FastGeoStreaming/Source/FastGeoStreaming/Private/FastGeoPrimitiveComponent.cpp
--- a/Engine/Plugins/Experimental/FastGeoStreaming/Source/FastGeoStreaming/Private/FastGeoPrimitiveComponent.cpp
+++ b/Engine/Plugins/Experimental/FastGeoStreaming/Source/FastGeoStreaming/Private/FastGeoPrimitiveComponent.cpp
@@ -15,6 +15,8 @@
 #include "WorldPartition/WorldPartition.h"
 #include "PhysicsEngine/BodyInstance.h"
 #include "PhysicsEngine/BodySetup.h"
+#include "PhysicsEngine/PhysicsObjectExternalInterface.h"
+#include "PhysicsProxy/SingleParticlePhysicsProxy.h"
 #include "PrimitiveComponentHelper.h"
 #include "PSOPrecache.h"
 #include "PSOPrecacheMaterial.h"
@@ -386,6 +388,21 @@ void FFastGeoPrimitiveComponent::DestroyRenderState(FFastGeoDestroyRenderStateCo
 	}
 }
 
+TArray<Chaos::FPhysicsObject*> FFastGeoPrimitiveComponent::GetAllPhysicsObjects() const
+{
+	auto GetPhysicsObject = [this]() -> Chaos::FPhysicsObject*
+	{
+		if (!BodyInstance.IsValidBodyInstance())
+		{
+			return nullptr;
+		}
+		return BodyInstance.GetPhysicsActor()->GetPhysicsObject();
+	};
+
+	TArray<Chaos::FPhysicsObject*> Bodies = { GetPhysicsObject() };
+	return Bodies;
+}
+
 void FFastGeoPrimitiveComponent::OnAsyncCreatePhysicsState()
 {
 	TRACE_CPUPROFILER_EVENT_SCOPE(FFastGeoPrimitiveComponent::OnAsyncCreatePhysicsState);
@@ -421,8 +438,19 @@ void FFastGeoPrimitiveComponent::OnAsyncCreatePhysicsState()
 			}
 #endif
 
-			// Create the body.
-			BodyInstance.InitBody(BodySetup, BodyTransform, nullptr, GetWorld()->GetPhysicsScene(), FInitBodySpawnParams(IsStaticPhysics(), /*bPhysicsTypeDeterminesSimulation*/false), this);
+			// Initialize BodyInstanceOwner
+			BodyInstanceOwner.Initialize(this);
+
+			// Initialize the body instance
+			BodyInstance.InitBody(BodySetup, BodyTransform, nullptr, GetWorld()->GetPhysicsScene(), FInitBodySpawnParams(IsStaticPhysics(), /*bPhysicsTypeDeterminesSimulation*/false), &BodyInstanceOwner);
+
+			// Assign BodyInstanceOwner
+			if (Chaos::FSingleParticlePhysicsProxy* ProxyHandle = BodyInstance.GetPhysicsActor())
+			{
+				TArray<Chaos::FPhysicsObject*> PhysicsObjects = GetAllPhysicsObjects();
+				FLockedWritePhysicsObjectExternalInterface Interface = FPhysicsObjectExternalInterface::LockWrite(PhysicsObjects);
+				Interface->SetUserDefinedEntity(PhysicsObjects, &BodyInstanceOwner);
+			}
 		}
 	}
 }
@@ -431,6 +459,13 @@ void FFastGeoPrimitiveComponent::OnAsyncDestroyPhysicsStateBegin_GameThread()
 {
 	TRACE_CPUPROFILER_EVENT_SCOPE(FFastGeoPrimitiveComponent::OnAsyncDestroyPhysicsStateBegin_GameThread);
 
+	// Remove all user defined entities here
+	TArray<Chaos::FPhysicsObject*> PhysicsObjects = GetAllPhysicsObjects();
+	FPhysicsObjectExternalInterface::LockWrite(PhysicsObjects)->SetUserDefinedEntity(PhysicsObjects, nullptr);
+
+	// Reset BodyInstanceOwner
+	BodyInstanceOwner.Uninitialize();
+
 	check(!AsyncTermBodyPayload.IsSet());
 	AsyncTermBodyPayload = BodyInstance.StartAsyncTermBody_GameThread();
 	check(!BodyInstance.IsValidBodyInstance());
@@ -737,7 +772,6 @@ EHasCustomNavigableGeometry::Type FFastGeoPrimitiveComponent::HasCustomNavigable
 	return bHasCustomNavigableGeometry;
 }
 
-//~ Begin IPhysicsBodyInstanceOwner interface
 bool FFastGeoPrimitiveComponent::IsStaticPhysics() const
 {
 	return GetSceneProxyDesc().Mobility != EComponentMobility::Movable && bStaticWhenNotMoveable;
@@ -747,13 +781,83 @@ UObject* FFastGeoPrimitiveComponent::GetSourceObject() const
 {
 	return GetOwnerContainer();
 }
-//~ End IPhysicsBodyInstanceOwner interface
+
+ECollisionResponse FFastGeoPrimitiveComponent::GetCollisionResponseToChannel(ECollisionChannel Channel) const
+{
+	return BodyInstance.GetResponseToChannel(Channel);
+}
 
 FPrimitiveMaterialPropertyDescriptor FFastGeoPrimitiveComponent::GetUsedMaterialPropertyDesc(ERHIFeatureLevel::Type FeatureLevel) const
 {
 	return FPrimitiveComponentHelper::GetUsedMaterialPropertyDesc(*this, FeatureLevel);
 }
 
+const FName FFastGeoPhysicsBodyInstanceOwner::NAME_FastGeoPhysicsBodyInstanceOwner = TEXT("FastGeoPhysicsBodyInstanceOwner");
+
+FFastGeoPhysicsBodyInstanceOwner::FFastGeoPhysicsBodyInstanceOwner()
+	: FChaosUserDefinedEntity(FFastGeoPhysicsBodyInstanceOwner::NAME_FastGeoPhysicsBodyInstanceOwner)
+	, OwnerComponent(nullptr)
+{}
+
+void FFastGeoPhysicsBodyInstanceOwner::Uninitialize()
+{
+	Initialize(nullptr);
+}
+
+void FFastGeoPhysicsBodyInstanceOwner::Initialize(FFastGeoPrimitiveComponent* InOwnerComponent)
+{
+	OwnerComponent = InOwnerComponent;
+	UFastGeoContainer* NewContainer = OwnerComponent ? OwnerComponent->GetOwnerContainer() : nullptr;
+	check(!OwnerComponent || NewContainer);
+	check(OwnerContainer.IsExplicitlyNull() || !NewContainer || OwnerContainer == NewContainer);
+	OwnerContainer = NewContainer;
+}
+
+TWeakObjectPtr<UObject> FFastGeoPhysicsBodyInstanceOwner::GetOwnerObject()
+{
+	return OwnerContainer;
+}
+
+IPhysicsBodyInstanceOwner* FFastGeoPhysicsBodyInstanceOwner::GetPhysicsBodyInstanceOwner(FChaosUserDefinedEntity* InUserDefinedEntity)
+{
+	if (InUserDefinedEntity && InUserDefinedEntity->GetEntityTypeName() == NAME_FastGeoPhysicsBodyInstanceOwner)
+	{
+		FFastGeoPhysicsBodyInstanceOwner* FastGeoPhysicsBodyInstanceOwner = static_cast<FFastGeoPhysicsBodyInstanceOwner*>(InUserDefinedEntity);
+		check(FastGeoPhysicsBodyInstanceOwner->GetOwnerObject().IsValid());
+		return FastGeoPhysicsBodyInstanceOwner;
+	}
+	return nullptr;
+}
+
+bool FFastGeoPhysicsBodyInstanceOwner::IsStaticPhysics() const
+{
+	check(OwnerContainer.IsValid());
+	return OwnerComponent->IsStaticPhysics();
+}
+
+UObject* FFastGeoPhysicsBodyInstanceOwner::GetSourceObject() const
+{
+	check(OwnerContainer.IsValid());
+	return OwnerComponent->GetSourceObject();
+}
+
+ECollisionResponse FFastGeoPhysicsBodyInstanceOwner::GetCollisionResponseToChannel(ECollisionChannel Channel) const 
+{
+	check(OwnerContainer.IsValid());
+	return OwnerComponent->GetCollisionResponseToChannel(Channel);
+}
+
+UPhysicalMaterial* FFastGeoPhysicsBodyInstanceOwner::GetPhysicalMaterial() const
+{
+	check(OwnerContainer.IsValid());
+	return OwnerComponent->GetPhysicalMaterial();
+}
+
+void FFastGeoPhysicsBodyInstanceOwner::GetComplexPhysicalMaterials(TArray<UPhysicalMaterial*>& OutPhysMaterials, TArray<FPhysicalMaterialMaskParams>* OutPhysMaterialMasks) const
+{
+	check(OwnerContainer.IsValid());
+	return OwnerComponent->GetComplexPhysicalMaterials(OutPhysMaterials, OutPhysMaterialMasks);
+}
 
 #if WITH_EDITOR
 void UFastGeoPrimitiveComponentEditorProxy::NotifyRenderStateChanged()
diff --git a/Engine/Plugins/Experimental/FastGeoStreaming/Source/FastGeoStreaming/Private/FastGeoPrimitiveComponent.h b/Engine/Plugins/Experimental/FastGeoStreaming/Source/FastGeoStreaming/Private/FastGeoPrimitiveComponent.h
--- a/Engine/Plugins/Experimental/FastGeoStreaming/Source/FastGeoStreaming/Private/FastGeoPrimitiveComponent.h
+++ b/Engine/Plugins/Experimental/FastGeoStreaming/Source/FastGeoStreaming/Private/FastGeoPrimitiveComponent.h
@@ -11,6 +11,7 @@
 #include "RenderCommandFence.h"
 #include "PhysicsEngine/BodyInstance.h"
 #include "PhysicsEngine/PhysicsBodyInstanceOwnerInterface.h"
+#include "Chaos/ChaosUserEntity.h"
 #include "PSOPrecacheFwd.h"
 #include "FastGeoPrimitiveComponent.generated.h"
 
@@ -30,7 +31,40 @@ enum class ESceneProxyCreationError
 	InvalidMesh
 };
 
-class FASTGEOSTREAMING_API FFastGeoPrimitiveComponent : public FFastGeoComponent, public IPhysicsBodyInstanceOwner
+class FFastGeoPhysicsBodyInstanceOwner : public FChaosUserDefinedEntity, public IPhysicsBodyInstanceOwner
+{
+public:
+	FFastGeoPhysicsBodyInstanceOwner();
+	virtual ~FFastGeoPhysicsBodyInstanceOwner() = default;
+
+	/** Returns the IPhysicsBodyInstanceOwner based on the provided FChaosUserDefinedEntity */
+	static IPhysicsBodyInstanceOwner* GetPhysicsBodyInstanceOwner(FChaosUserDefinedEntity* InUserDefinedEntity);
+
+	//~ Begin FChaosUserDefinedEntity interface
+	TWeakObjectPtr<UObject> GetOwnerObject() override;
+	//~ End FChaosUserDefinedEntity interface
+
+	//~ Begin IPhysicsBodyInstanceOwner interface
+	virtual bool IsStaticPhysics() const override;
+	virtual UObject* GetSourceObject() const override;
+	virtual ECollisionResponse GetCollisionResponseToChannel(ECollisionChannel Channel) const override;
+	virtual UPhysicalMaterial* GetPhysicalMaterial() const override;
+	virtual void GetComplexPhysicalMaterials(TArray<UPhysicalMaterial*>& OutPhysMaterials, TArray<FPhysicalMaterialMaskParams>* OutPhysMaterialMasks) const override;
+	//~ End IPhysicsBodyInstanceOwner interface
+
+private:
+	void Uninitialize();
+	void Initialize(FFastGeoPrimitiveComponent* InOwner);
+
+	FFastGeoPrimitiveComponent* OwnerComponent;
+	TWeakObjectPtr<UFastGeoContainer> OwnerContainer;
+	static const FName NAME_FastGeoPhysicsBodyInstanceOwner;
+
+	friend class FFastGeoPrimitiveComponent;
+	friend class FFastGeoInstancedStaticMeshComponent;
+};
+
+class FASTGEOSTREAMING_API FFastGeoPrimitiveComponent : public FFastGeoComponent
 {
 public:
 	typedef FFastGeoComponent Super;
@@ -85,13 +119,6 @@ public:
 	void UpdateVisibility();
 	EComponentMobility::Type GetMobility() const;
 	
-	//~ Begin IPhysicsBodyInstanceOwner interface
-	virtual bool IsStaticPhysics() const override;
-	virtual UObject* GetSourceObject() const override;
-	virtual UPhysicalMaterial* GetPhysicalMaterial() const override { return nullptr; }
-	virtual void GetComplexPhysicalMaterials(TArray<UPhysicalMaterial*>& OutPhysMaterials, TArray<FPhysicalMaterialMaskParams>* OutPhysMaterialMasks) const override {}
-	//~ End IPhysicsBodyInstanceOwner interface
-	
 	//~ Begin FFastGeoComponent interface
 	virtual void Serialize(FArchive& Ar) override;
 	virtual void InitializeDynamicProperties() override;
@@ -134,6 +161,15 @@ protected:
 	virtual FPrimitiveSceneProxy* CreateSceneProxy(ESceneProxyCreationError* OutError = nullptr) = 0;
 	FPrimitiveSceneDesc BuildSceneDesc();
 
+	//~ Begin Physics
+	virtual TArray<Chaos::FPhysicsObject*> GetAllPhysicsObjects() const;
+	virtual bool IsStaticPhysics() const;
+	virtual UObject* GetSourceObject() const;
+	virtual UPhysicalMaterial* GetPhysicalMaterial() const { return nullptr; }
+	virtual void GetComplexPhysicalMaterials(TArray<UPhysicalMaterial*>& OutPhysMaterials, TArray<FPhysicalMaterialMaskParams>* OutPhysMaterialMasks) const {}
+	virtual ECollisionResponse GetCollisionResponseToChannel(ECollisionChannel Channel) const;
+	//~ End Navigation
+
 protected:
 	// Persistent data
 	FTransform LocalTransform;
@@ -152,6 +188,7 @@ protected:
 	TArray<URuntimeVirtualTexture*> RuntimeVirtualTextures;
 
 	// Runtime Data (transient)
+	FFastGeoPhysicsBodyInstanceOwner BodyInstanceOwner;
 	FPrimitiveSceneInfoData PrimitiveSceneData{};
 	/** Payload used to release BodyInstance resources in asynchronous mode (see OnAsyncDestroyPhysicsState). */
 	TOptional<FBodyInstance::FAsyncTermBodyPayload> AsyncTermBodyPayload;
@@ -194,6 +231,7 @@ protected:
 #endif
 
 	friend class FPrimitiveComponentHelper;
+	friend class FFastGeoPhysicsBodyInstanceOwner;
 	friend class UFastGeoContainer;
 	friend class UFastGeoPrimitiveComponentEditorProxy;
 };
diff --git a/Engine/Plugins/Experimental/FastGeoStreaming/Source/FastGeoStreaming/Private/FastGeoStaticMeshComponent.h b/Engine/Plugins/Experimental/FastGeoStreaming/Source/FastGeoStreaming/Private/FastGeoStaticMeshComponent.h
--- a/Engine/Plugins/Experimental/FastGeoStreaming/Source/FastGeoStreaming/Private/FastGeoStaticMeshComponent.h
+++ b/Engine/Plugins/Experimental/FastGeoStreaming/Source/FastGeoStreaming/Private/FastGeoStaticMeshComponent.h
@@ -61,14 +61,6 @@ public:
 	virtual FBox GetNavigationBounds() const override;
 	virtual void GetNavigationData(FNavigationRelevantData& Data) const override;
 	virtual bool DoCustomNavigableGeometryExport(FNavigableGeometryExport& GeomExport) const override;
-	//~ End FFastGeoPrimitiveComponent interface
-
-	//~ Begin IPhysicsBodyInstanceOwner interface
-	virtual UPhysicalMaterial* GetPhysicalMaterial() const override;
-	virtual void GetComplexPhysicalMaterials(TArray<UPhysicalMaterial*>& OutPhysMaterials, TArray<FPhysicalMaterialMaskParams>* OutPhysMaterialMasks) const override;
-	//~ End IPhysicsBodyInstanceOwner interface
-
-	//~ Begin FFastGeoPrimitiveComponent interface
 	virtual UMaterialInterface* GetMaterial(int32 MaterialIndex) const override;
 	virtual int32 GetNumMaterials() const override;
 	virtual void GetUsedMaterials(TArray<UMaterialInterface*>& OutMaterials, bool bGetDebugMaterials = false) const override;
@@ -83,6 +75,8 @@ protected:
 	virtual void InitializeSceneProxyDescDynamicProperties() override;
 	virtual void ApplyWorldTransform(const FTransform& InTransform) override;
 	virtual FPrimitiveSceneProxy* CreateSceneProxy(ESceneProxyCreationError* OutError = nullptr) override;
+	virtual UPhysicalMaterial* GetPhysicalMaterial() const override;
+	virtual void GetComplexPhysicalMaterials(TArray<UPhysicalMaterial*>& OutPhysMaterials, TArray<FPhysicalMaterialMaskParams>* OutPhysMaterialMasks) const override;
 	//~ End FFastGeoPrimitiveComponent interface
 
 	//~ Begin FFastGeoMeshComponent interface
diff --git a/Engine/Source/Runtime/Engine/Classes/PhysicsEngine/PhysicsBodyInstanceOwnerInterface.h b/Engine/Source/Runtime/Engine/Classes/PhysicsEngine/PhysicsBodyInstanceOwnerInterface.h
--- a/Engine/Source/Runtime/Engine/Classes/PhysicsEngine/PhysicsBodyInstanceOwnerInterface.h
+++ b/Engine/Source/Runtime/Engine/Classes/PhysicsEngine/PhysicsBodyInstanceOwnerInterface.h
@@ -2,16 +2,46 @@
 
 #pragma once
 
+#include "UObject/Interface.h"
 #include "Containers/Array.h"
+#include "Chaos/PhysicsObject.h"
+#include "PhysicsBodyInstanceOwnerInterface.generated.h"
 
 class UObject;
 class UPhysicalMaterial;
+class IPhysicsBodyInstanceOwner;
+struct FHitResult;
+struct FOverlapResult;
 struct FPhysicalMaterialMaskParams;
+enum ECollisionResponse : int;
+enum ECollisionChannel : int;
+
+UINTERFACE(MinimalAPI)
+class UPhysicsBodyInstanceOwnerResolver : public UInterface
+{
+	GENERATED_BODY()
+};
+
+class IPhysicsBodyInstanceOwnerResolver
+{
+	GENERATED_BODY()
+
+public:
+	virtual IPhysicsBodyInstanceOwner* ResolvePhysicsBodyInstanceOwner(Chaos::FConstPhysicsObjectHandle PhysicsObject) = 0;
+};
 
 /** Interface representing the owner of a FBodyInstance (used when the owner is not an UPrimitiveComponent). */
 class IPhysicsBodyInstanceOwner
 {
 public:
+	virtual ~IPhysicsBodyInstanceOwner() = default;
+
+	/** Returns the IPhysicsBodyInstanceOwner based on a given hit result. */
+	static ENGINE_API IPhysicsBodyInstanceOwner* GetPhysicsBodyInstandeOwnerFromHitResult(const FHitResult& Result);
+
+	/** Returns the IPhysicsBodyInstanceOwner based on a given overlap result. */
+	static ENGINE_API IPhysicsBodyInstanceOwner* GetPhysicsBodyInstandeOwnerFromOverlapResult(const FOverlapResult& OverlapResult);
+
 	/** Whether the physics is static. */
 	virtual bool IsStaticPhysics() const = 0;
 
@@ -23,4 +53,7 @@ public:
 
 	/** Get the complex PhysicalMaterials array for this body (used by FBodyInstance::GetComplexPhysicalMaterials). */
 	virtual void GetComplexPhysicalMaterials(TArray<UPhysicalMaterial*>& OutPhysMaterials, TArray<FPhysicalMaterialMaskParams>* OutPhysMaterialMasks) const = 0;
+
+	/** Gets the response type given a specific channel. */
+	virtual ECollisionResponse GetCollisionResponseToChannel(ECollisionChannel Channel) const = 0;
 };
\ No newline at end of file
diff --git a/Engine/Source/Runtime/Engine/Private/Collision/CollisionConversions.cpp b/Engine/Source/Runtime/Engine/Private/Collision/CollisionConversions.cpp
--- a/Engine/Source/Runtime/Engine/Private/Collision/CollisionConversions.cpp
+++ b/Engine/Source/Runtime/Engine/Private/Collision/CollisionConversions.cpp
@@ -137,7 +137,6 @@ static void SetHitResultFromShapeAndFaceIndex(const FPhysicsShape& Shape,  const
 	CHAOS_CHECK(ShapeIndex < (int32)TNumericLimits<uint8>::Max()); // I could just write < 256, but this makes it more clear *why*
 	OutResult.ElementIndex = (uint8)ShapeIndex;
 
-	IPhysicsBodyInstanceOwner* BodyInstanceOwner = nullptr;
 	TWeakObjectPtr<UPrimitiveComponent> OwningComponent;
 	OutResult.PhysicsObject = nullptr;
 	
@@ -158,7 +157,6 @@ static void SetHitResultFromShapeAndFaceIndex(const FPhysicsShape& Shape,  const
 		}
 
 		OwningComponent = BodyInst->OwnerComponent;
-		BodyInstanceOwner = BodyInst->BodyInstanceOwner;
 		OutResult.PhysicsObject = BodyInst->GetPhysicsActor() ? BodyInst->GetPhysicsActor()->GetPhysicsObject() : nullptr;
 	}
 	else
@@ -204,7 +202,7 @@ static void SetHitResultFromShapeAndFaceIndex(const FPhysicsShape& Shape,  const
 	OutResult.PhysMaterial = nullptr;
 
 	// Grab actor/component
-	if (!OwningComponent.IsExplicitlyNull() || BodyInstanceOwner)
+	if (!OwningComponent.IsExplicitlyNull() || IPhysicsBodyInstanceOwner::GetPhysicsBodyInstandeOwnerFromHitResult(OutResult))
 	{
 		if (!OwningComponent.IsExplicitlyNull())
 		{
diff --git a/Engine/Source/Runtime/Engine/Private/PhysicsEngine/PhysicsBodyInstanceOwnerInterface.cpp b/Engine/Source/Runtime/Engine/Private/PhysicsEngine/PhysicsBodyInstanceOwnerInterface.cpp
new file mode 100644
--- /dev/null
+++ b/Engine/Source/Runtime/Engine/Private/PhysicsEngine/PhysicsBodyInstanceOwnerInterface.cpp
@@ -0,0 +1,23 @@
+// Copyright Epic Games, Inc. All Rights Reserved.
+
+#include "PhysicsEngine/PhysicsBodyInstanceOwnerInterface.h"
+#include "Engine/HitResult.h"
+#include "Engine/OverlapResult.h"
+
+IPhysicsBodyInstanceOwner* IPhysicsBodyInstanceOwner::GetPhysicsBodyInstandeOwnerFromHitResult(const FHitResult& Result)
+{
+	if (IPhysicsBodyInstanceOwnerResolver* PhysicsObjectOwnerResolver = Result.PhysicsObject ? Cast<IPhysicsBodyInstanceOwnerResolver>(Result.PhysicsObjectOwner.Get()) : nullptr)
+	{
+		return PhysicsObjectOwnerResolver->ResolvePhysicsBodyInstanceOwner(Result.PhysicsObject);
+	}
+	return nullptr;
+}
+
+IPhysicsBodyInstanceOwner* IPhysicsBodyInstanceOwner::GetPhysicsBodyInstandeOwnerFromOverlapResult(const FOverlapResult& Result)
+{
+	if (IPhysicsBodyInstanceOwnerResolver* PhysicsObjectOwnerResolver = Result.PhysicsObject ? Cast<IPhysicsBodyInstanceOwnerResolver>(Result.PhysicsObjectOwner.Get()) : nullptr)
+	{
+		return PhysicsObjectOwnerResolver->ResolvePhysicsBodyInstanceOwner(Result.PhysicsObject);
+	}
+	return nullptr;
+}
\ No newline at end of file
