FastGeo: When destroying physics, the UserDefinedEntity of the physics objects are now reset asynchronously

diff --git a/Engine/Plugins/Experimental/FastGeoStreaming/Source/FastGeoStreaming/Private/FastGeoComponent.h b/Engine/Plugins/Experimental/FastGeoStreaming/Source/FastGeoStreaming/Private/FastGeoComponent.h
--- a/Engine/Plugins/Experimental/FastGeoStreaming/Source/FastGeoStreaming/Private/FastGeoComponent.h
+++ b/Engine/Plugins/Experimental/FastGeoStreaming/Source/FastGeoStreaming/Private/FastGeoComponent.h
@@ -41,8 +41,8 @@ public:
 	virtual void OnAsyncCreatePhysicsState();
 	virtual void OnAsyncDestroyPhysicsStateBegin_GameThread();
 	virtual void OnAsyncDestroyPhysicsState();
+	virtual void OnAsyncDestroyPhysicsStateEnd_GameThread();
 	void OnAsyncCreatePhysicsStateEnd_GameThread();
-	void OnAsyncDestroyPhysicsStateEnd_GameThread();
 
 	FFastGeoComponentCluster* GetOwnerComponentCluster() const;
 	UFastGeoContainer* GetOwnerContainer() const;
diff --git a/Engine/Plugins/Experimental/FastGeoStreaming/Source/FastGeoStreaming/Private/FastGeoInstancedStaticMeshComponent.cpp b/Engine/Plugins/Experimental/FastGeoStreaming/Source/FastGeoStreaming/Private/FastGeoInstancedStaticMeshComponent.cpp
--- a/Engine/Plugins/Experimental/FastGeoStreaming/Source/FastGeoStreaming/Private/FastGeoInstancedStaticMeshComponent.cpp
+++ b/Engine/Plugins/Experimental/FastGeoStreaming/Source/FastGeoStreaming/Private/FastGeoInstancedStaticMeshComponent.cpp
@@ -244,24 +244,29 @@ void FFastGeoInstancedStaticMeshComponent::OnAsyncDestroyPhysicsStateBegin_GameT
 
 	FFastGeoComponent::OnAsyncDestroyPhysicsStateBegin_GameThread();
 
-	// Remove all user defined entities here
-	TArray<Chaos::FPhysicsObject*> PhysicsObjects = GetAllPhysicsObjects();
-	FPhysicsObjectExternalInterface::LockWrite(PhysicsObjects)->SetUserDefinedEntity(PhysicsObjects, nullptr);
-
-	// Reset BodyInstanceOwner
-	BodyInstanceOwner.Uninitialize();
-
 	// Move InstanceBodies in AsyncDestroyPhysicsStatePayload
 	check(AsyncDestroyPhysicsStatePayload.IsEmpty());
 	AsyncDestroyPhysicsStatePayload = MoveTemp(InstanceBodies);
 }
 
+void FFastGeoInstancedStaticMeshComponent::OnAsyncDestroyPhysicsStateEnd_GameThread()
+{
+	FFastGeoComponent::OnAsyncDestroyPhysicsStateEnd_GameThread();
+
+	// Reset BodyInstanceOwner
+	BodyInstanceOwner.Uninitialize();
+}
+
 void FFastGeoInstancedStaticMeshComponent::OnAsyncDestroyPhysicsState()
 {
 	TRACE_CPUPROFILER_EVENT_SCOPE(FFastGeoInstancedStaticMeshComponent::OnAsyncDestroyPhysicsState);
 
 	FFastGeoComponent::OnAsyncDestroyPhysicsState();
 
+	// Remove all user defined entities
+	TArray<Chaos::FPhysicsObject*> PhysicsObjects = GetAllPhysicsObjects(AsyncDestroyPhysicsStatePayload);
+	FPhysicsObjectExternalInterface::LockWrite(PhysicsObjects)->SetUserDefinedEntity(PhysicsObjects, nullptr);
+
 	check(InstanceBodies.IsEmpty());
 	for (FBodyInstance*& Instance : AsyncDestroyPhysicsStatePayload)
 	{
@@ -274,22 +279,16 @@ void FFastGeoInstancedStaticMeshComponent::OnAsyncDestroyPhysicsState()
 	AsyncDestroyPhysicsStatePayload.Empty();
 }
 
-TArray<Chaos::FPhysicsObject*> FFastGeoInstancedStaticMeshComponent::GetAllPhysicsObjects() const
+TArray<Chaos::FPhysicsObject*> FFastGeoInstancedStaticMeshComponent::GetAllPhysicsObjects(TArray<FBodyInstance*> InInstanceBodies)
 {
-	auto GetPhysicsObjectById = [this](Chaos::FPhysicsObjectId Id) -> Chaos::FPhysicsObject*
+	TArray<Chaos::FPhysicsObject*> Objects;
+	Objects.Reserve(InInstanceBodies.Num());
+	for (FBodyInstance* InstancedBody : InInstanceBodies)
 	{
-		if (!InstanceBodies.IsValidIndex(Id) || !InstanceBodies[Id] || !InstanceBodies[Id]->GetPhysicsActor())
+		if (Chaos::FPhysicsObject* PhysicsObject = InstancedBody && InstancedBody->GetPhysicsActor() ? InstancedBody->GetPhysicsActor()->GetPhysicsObject() : nullptr)
 		{
-			return nullptr;
+			Objects.Add(PhysicsObject);
 		}
-		return InstanceBodies[Id]->GetPhysicsActor()->GetPhysicsObject();
-	};
-
-	TArray<Chaos::FPhysicsObject*> Objects;
-	Objects.Reserve(InstanceBodies.Num());
-	for (int32 Index = 0; Index < InstanceBodies.Num(); ++Index)
-	{
-		Objects.Add(GetPhysicsObjectById(Index));
 	}
 	return Objects;
 }
@@ -349,9 +348,8 @@ void FFastGeoInstancedStaticMeshComponent::CreateAllInstanceBodies()
 			FBodyInstance::InitStaticBodies(MoveTemp(InstanceBodiesSanitized), MoveTemp(Transforms), BodySetup, nullptr, GetWorld()->GetPhysicsScene(), &BodyInstanceOwner);
 
 			// Assign BodyInstanceOwner
-			TArray<Chaos::FPhysicsObject*> PhysicsObjects = GetAllPhysicsObjects();
-			FLockedWritePhysicsObjectExternalInterface Interface = FPhysicsObjectExternalInterface::LockWrite(PhysicsObjects);
-			Interface->SetUserDefinedEntity(PhysicsObjects, &BodyInstanceOwner);
+			TArray<Chaos::FPhysicsObject*> PhysicsObjects = GetAllPhysicsObjects(InstanceBodies);
+			FPhysicsObjectExternalInterface::LockWrite(PhysicsObjects)->SetUserDefinedEntity(PhysicsObjects, &BodyInstanceOwner);
 		}
 	}
 	else
diff --git a/Engine/Plugins/Experimental/FastGeoStreaming/Source/FastGeoStreaming/Private/FastGeoInstancedStaticMeshComponent.h b/Engine/Plugins/Experimental/FastGeoStreaming/Source/FastGeoStreaming/Private/FastGeoInstancedStaticMeshComponent.h
--- a/Engine/Plugins/Experimental/FastGeoStreaming/Source/FastGeoStreaming/Private/FastGeoInstancedStaticMeshComponent.h
+++ b/Engine/Plugins/Experimental/FastGeoStreaming/Source/FastGeoStreaming/Private/FastGeoInstancedStaticMeshComponent.h
@@ -22,6 +22,7 @@ public:
 	virtual void OnAsyncCreatePhysicsState() override;
 	virtual void OnAsyncDestroyPhysicsState() override;
 	virtual void OnAsyncDestroyPhysicsStateBegin_GameThread() override;
+	virtual void OnAsyncDestroyPhysicsStateEnd_GameThread() override;
 #if WITH_EDITOR
 	virtual void InitializeFromComponent(UActorComponent* Component) override;
 	virtual UClass* GetEditorProxyClass() const override;
@@ -34,7 +35,6 @@ public:
 	virtual bool IsNavigationRelevant() const override;
 	virtual bool DoCustomNavigableGeometryExport(FNavigableGeometryExport& GeomExport) const override;
 	virtual bool ShouldSkipNavigationDirtyAreaOnAddOrRemove() override { return true; }
-	virtual TArray<Chaos::FPhysicsObject*> GetAllPhysicsObjects() const override;
 	//~ End FFastGeoPrimitiveComponent interface
 
 protected:
@@ -68,9 +68,9 @@ private:
 	};
 	FBoxSphereBounds CalculateBounds(EBoundsType BoundsType);
 	void CreateAllInstanceBodies();
+	static TArray<Chaos::FPhysicsObject*> GetAllPhysicsObjects(TArray<FBodyInstance*> InInstanceBodies);
 
 private:
-
 	// Persistent Data
 	TArray<FInstancedStaticMeshInstanceData> PerInstanceSMData;
 	int32 InstancingRandomSeed = 0;
diff --git a/Engine/Plugins/Experimental/FastGeoStreaming/Source/FastGeoStreaming/Private/FastGeoPrimitiveComponent.cpp b/Engine/Plugins/Experimental/FastGeoStreaming/Source/FastGeoStreaming/Private/FastGeoPrimitiveComponent.cpp
--- a/Engine/Plugins/Experimental/FastGeoStreaming/Source/FastGeoStreaming/Private/FastGeoPrimitiveComponent.cpp
+++ b/Engine/Plugins/Experimental/FastGeoStreaming/Source/FastGeoStreaming/Private/FastGeoPrimitiveComponent.cpp
@@ -388,21 +388,6 @@ void FFastGeoPrimitiveComponent::DestroyRenderState(FFastGeoDestroyRenderStateCo
 	}
 }
 
-TArray<Chaos::FPhysicsObject*> FFastGeoPrimitiveComponent::GetAllPhysicsObjects() const
-{
-	auto GetPhysicsObject = [this]() -> Chaos::FPhysicsObject*
-	{
-		if (!BodyInstance.IsValidBodyInstance())
-		{
-			return nullptr;
-		}
-		return BodyInstance.GetPhysicsActor()->GetPhysicsObject();
-	};
-
-	TArray<Chaos::FPhysicsObject*> Bodies = { GetPhysicsObject() };
-	return Bodies;
-}
-
 void FFastGeoPrimitiveComponent::OnAsyncCreatePhysicsState()
 {
 	TRACE_CPUPROFILER_EVENT_SCOPE(FFastGeoPrimitiveComponent::OnAsyncCreatePhysicsState);
@@ -447,9 +432,11 @@ void FFastGeoPrimitiveComponent::OnAsyncCreatePhysicsState()
 			// Assign BodyInstanceOwner
 			if (Chaos::FSingleParticlePhysicsProxy* ProxyHandle = BodyInstance.GetPhysicsActor())
 			{
-				TArray<Chaos::FPhysicsObject*> PhysicsObjects = GetAllPhysicsObjects();
-				FLockedWritePhysicsObjectExternalInterface Interface = FPhysicsObjectExternalInterface::LockWrite(PhysicsObjects);
-				Interface->SetUserDefinedEntity(PhysicsObjects, &BodyInstanceOwner);
+				if (Chaos::FPhysicsObject* PhysicsObject = BodyInstance.IsValidBodyInstance() ? BodyInstance.GetPhysicsActor()->GetPhysicsObject() : nullptr)
+				{
+					TArrayView<Chaos::FPhysicsObject*> PhysicsObjects(&PhysicsObject, 1);
+					FPhysicsObjectExternalInterface::LockWrite(PhysicsObjects)->SetUserDefinedEntity(PhysicsObjects, &BodyInstanceOwner);
+				}
 			}
 		}
 	}
@@ -459,13 +446,6 @@ void FFastGeoPrimitiveComponent::OnAsyncDestroyPhysicsStateBegin_GameThread()
 {
 	TRACE_CPUPROFILER_EVENT_SCOPE(FFastGeoPrimitiveComponent::OnAsyncDestroyPhysicsStateBegin_GameThread);
 
-	// Remove all user defined entities here
-	TArray<Chaos::FPhysicsObject*> PhysicsObjects = GetAllPhysicsObjects();
-	FPhysicsObjectExternalInterface::LockWrite(PhysicsObjects)->SetUserDefinedEntity(PhysicsObjects, nullptr);
-
-	// Reset BodyInstanceOwner
-	BodyInstanceOwner.Uninitialize();
-
 	check(!AsyncTermBodyPayload.IsSet());
 	AsyncTermBodyPayload = BodyInstance.StartAsyncTermBody_GameThread();
 	check(!BodyInstance.IsValidBodyInstance());
@@ -473,6 +453,14 @@ void FFastGeoPrimitiveComponent::OnAsyncDestroyPhysicsStateBegin_GameThread()
 	Super::OnAsyncDestroyPhysicsStateBegin_GameThread();
 }
 
+void FFastGeoPrimitiveComponent::OnAsyncDestroyPhysicsStateEnd_GameThread()
+{
+	Super::OnAsyncDestroyPhysicsStateEnd_GameThread();
+
+	// Reset BodyInstanceOwner
+	BodyInstanceOwner.Uninitialize();
+}
+
 void FFastGeoPrimitiveComponent::OnAsyncDestroyPhysicsState()
 {
 	TRACE_CPUPROFILER_EVENT_SCOPE(FFastGeoPrimitiveComponent::OnAsyncDestroyPhysicsState);
@@ -480,6 +468,13 @@ void FFastGeoPrimitiveComponent::OnAsyncDestroyPhysicsState()
 	// We tell the BodyInstance to shut down the physics-engine data.
 	if (ensure(AsyncTermBodyPayload.IsSet()))
 	{
+		// Remove all user defined entities
+		if (Chaos::FPhysicsObject* PhysicsObject = AsyncTermBodyPayload->GetPhysicsActor() ? AsyncTermBodyPayload->GetPhysicsActor()->GetPhysicsObject() : nullptr)
+		{
+			TArrayView<Chaos::FPhysicsObject*> PhysicsObjects(&PhysicsObject, 1);
+			FPhysicsObjectExternalInterface::LockWrite(PhysicsObjects)->SetUserDefinedEntity(PhysicsObjects, nullptr);
+		}
+
 		FBodyInstance::AsyncTermBody(AsyncTermBodyPayload.GetValue());
 		AsyncTermBodyPayload.Reset();
 	}
@@ -1010,4 +1005,4 @@ HHitProxy* UFastGeoPrimitiveComponentEditorProxy::CreatePrimitiveHitProxies(TArr
 	return nullptr;
 }
 //~ End IPrimitiveComponent interface
-#endif
+#endif
\ No newline at end of file
diff --git a/Engine/Plugins/Experimental/FastGeoStreaming/Source/FastGeoStreaming/Private/FastGeoPrimitiveComponent.h b/Engine/Plugins/Experimental/FastGeoStreaming/Source/FastGeoStreaming/Private/FastGeoPrimitiveComponent.h
--- a/Engine/Plugins/Experimental/FastGeoStreaming/Source/FastGeoStreaming/Private/FastGeoPrimitiveComponent.h
+++ b/Engine/Plugins/Experimental/FastGeoStreaming/Source/FastGeoStreaming/Private/FastGeoPrimitiveComponent.h
@@ -124,6 +124,7 @@ public:
 	virtual void InitializeDynamicProperties() override;
 	virtual void OnAsyncCreatePhysicsState() override;
 	virtual void OnAsyncDestroyPhysicsStateBegin_GameThread() override;
+	virtual void OnAsyncDestroyPhysicsStateEnd_GameThread() override;
 	virtual void OnAsyncDestroyPhysicsState() override;
 	virtual bool IsCollisionEnabled() const override;
 #if WITH_EDITOR
@@ -162,7 +163,6 @@ protected:
 	FPrimitiveSceneDesc BuildSceneDesc();
 
 	//~ Begin Physics
-	virtual TArray<Chaos::FPhysicsObject*> GetAllPhysicsObjects() const;
 	virtual bool IsStaticPhysics() const;
 	virtual UObject* GetSourceObject() const;
 	virtual UPhysicalMaterial* GetPhysicalMaterial() const { return nullptr; }
diff --git a/Engine/Source/Runtime/Engine/Classes/PhysicsEngine/BodyInstance.h b/Engine/Source/Runtime/Engine/Classes/PhysicsEngine/BodyInstance.h
--- a/Engine/Source/Runtime/Engine/Classes/PhysicsEngine/BodyInstance.h
+++ b/Engine/Source/Runtime/Engine/Classes/PhysicsEngine/BodyInstance.h
@@ -801,6 +801,7 @@ public:
 	struct ENGINE_API FAsyncTermBodyPayload
 	{
 		FAsyncTermBodyPayload();
+		FPhysicsActorHandle GetPhysicsActor() const { return ActorHandle; }
 
 	private:
 		FPhysicsActorHandle ActorHandle;
