FastGeo: Fixed case where PSO precaching of a primitive becomes ready after removing the level from the world.

diff --git a/Engine/Plugins/Experimental/FastGeoStreaming/Source/FastGeoStreaming/Private/FastGeoPrimitiveComponent.cpp b/Engine/Plugins/Experimental/FastGeoStreaming/Source/FastGeoStreaming/Private/FastGeoPrimitiveComponent.cpp
--- a/Engine/Plugins/Experimental/FastGeoStreaming/Source/FastGeoStreaming/Private/FastGeoPrimitiveComponent.cpp
+++ b/Engine/Plugins/Experimental/FastGeoStreaming/Source/FastGeoStreaming/Private/FastGeoPrimitiveComponent.cpp
@@ -587,7 +587,11 @@ struct FFastGeoPSOPrecacheFinishedTask
 	{
 		if (FFastGeoPrimitiveComponent* PrimitiveComponent = WeakPrimitiveComponent.Get<FFastGeoPrimitiveComponent>())
 		{
-			PrimitiveComponent->OnPrecacheFinished(JobSetThatJustCompleted);
+			// Validate that the component is still part of a streamed-in level
+			if (PrimitiveComponent->GetWorld())
+			{
+				PrimitiveComponent->OnPrecacheFinished(JobSetThatJustCompleted);
+			}
 		}
 	}
 
