FastGeo - Improved the selection debug mode
* Should now work in PIE
* Should now report that non-spatially loaded actors can't be converted to fast geo. These do not ever get processed by the cell transformer under normal circumstances

diff --git a/Engine/Plugins/Experimental/FastGeoStreaming/Source/FastGeoStreaming/Private/FastGeoWorldPartitionRuntimeCellTransformer.cpp b/Engine/Plugins/Experimental/FastGeoStreaming/Source/FastGeoStreaming/Private/FastGeoWorldPartitionRuntimeCellTransformer.cpp
--- a/Engine/Plugins/Experimental/FastGeoStreaming/Source/FastGeoStreaming/Private/FastGeoWorldPartitionRuntimeCellTransformer.cpp
+++ b/Engine/Plugins/Experimental/FastGeoStreaming/Source/FastGeoStreaming/Private/FastGeoWorldPartitionRuntimeCellTransformer.cpp
@@ -179,7 +179,7 @@ bool UFastGeoWorldPartitionRuntimeCellTransformer::IsFastGeoEnabled = true;
 FAutoConsoleVariableRef UFastGeoWorldPartitionRuntimeCellTransformer::CVarIsFastGeoEnabled(
 	TEXT("FastGeo.Enable"),
 	UFastGeoWorldPartitionRuntimeCellTransformer::IsFastGeoEnabled,
-	TEXT("Set to false to disable FastGeoStreaming (used when starting PIE and at cook time)."),
+	TEXT("Set to false to disable FastGeoStreaming (used in PIE and at cook time)."),
 	ECVF_Default);
 
 void UFastGeoWorldPartitionRuntimeCellTransformer::BeginDestroy()
@@ -199,74 +199,69 @@ bool UFastGeoWorldPartitionRuntimeCellTransformer::IsDebugMode() const
 
 void UFastGeoWorldPartitionRuntimeCellTransformer::OnSelectionChanged(UObject* Object)
 {
-	if (bDebugModeOnEditorSelection && IsEnabled())
+	if (bDebugModeOnSelection && IsEnabled())
 	{
-		UWorldPartition* OuterWorldPartition = GetTypedOuter<UWorldPartition>();
-		UWorld* World = OuterWorldPartition ? OuterWorldPartition->GetWorld() : nullptr;
-		if (World && !World->IsGameWorld())
-		{
-			TArray<AActor*> SelectedActors;
+		TArray<AActor*> SelectedActors;
 
-			TFunction<void(AActor* InActor)> AddActorToSelection = [this, World, &SelectedActors, &AddActorToSelection](AActor* InActor)
+		TFunction<void(AActor* InActor)> AddActorToSelection = [this, &SelectedActors, &AddActorToSelection](AActor* InActor)
+		{
+			if (!CanAlwaysIgnoreActor(InActor))
 			{
-				if ((InActor->GetWorld() == World) && !CanAlwaysIgnoreActor(InActor))
+				const ILevelInstanceInterface* LevelInstance = Cast<ILevelInstanceInterface>(InActor);
+				if (LevelInstance && LevelInstance->GetDesiredRuntimeBehavior() == ELevelInstanceRuntimeBehavior::Partitioned && InActor->GetWorld())
 				{
-					const ILevelInstanceInterface* LevelInstance = Cast<ILevelInstanceInterface>(InActor);
-					if (LevelInstance && LevelInstance->GetDesiredRuntimeBehavior() == ELevelInstanceRuntimeBehavior::Partitioned && InActor->GetWorld())
+					if (const ULevelInstanceSubsystem* LevelInstanceSubsystem = InActor->GetWorld()->GetSubsystem<ULevelInstanceSubsystem>())
 					{
-						if (const ULevelInstanceSubsystem* LevelInstanceSubsystem = InActor->GetWorld()->GetSubsystem<ULevelInstanceSubsystem>())
+						LevelInstanceSubsystem->ForEachActorInLevelInstance(LevelInstance, [&AddActorToSelection](AActor* Actor)
 						{
-							LevelInstanceSubsystem->ForEachActorInLevelInstance(LevelInstance, [&AddActorToSelection](AActor* Actor)
-							{
-								AddActorToSelection(Actor);
-								return true;
-							});
-						}
-					}
-					else
-					{
-						SelectedActors.Add(InActor);
+							AddActorToSelection(Actor);
+							return true;
+						});
 					}
 				}
-			};
+				else
+				{
+					SelectedActors.Add(InActor);
+				}
+			}
+		};
 
-			if (USelection* Selection = Cast<USelection>(Object))
+		if (USelection* Selection = Cast<USelection>(Object))
+		{
+			for (int32 Index = 0; Index < Selection->Num(); Index++)
 			{
-				for (int32 Index = 0; Index < Selection->Num(); Index++)
+				if (AActor* SelectedActor = Cast<AActor>(Selection->GetSelectedObject(Index)))
 				{
-					if (AActor* SelectedActor = Cast<AActor>(Selection->GetSelectedObject(Index)))
-					{
-						AddActorToSelection(SelectedActor);
-					}
+					AddActorToSelection(SelectedActor);
 				}
 			}
+		}
 
-			if (!SelectedActors.IsEmpty())
-			{
-				TGuardValue<bool> Guard(FFastGeoTransformResult::ShouldReport, true);
+		if (!SelectedActors.IsEmpty())
+		{
+			TGuardValue<bool> Guard(FFastGeoTransformResult::ShouldReport, true);
 
-				UE_LOG(LogFastGeoStreaming, Log, TEXT("------------------------------------------------------------------------"));
-				UE_LOG(LogFastGeoStreaming, Log, TEXT("- FastGeoStreaming Debug Mode: Transformation on %d selected actors "), SelectedActors.Num());
+			UE_LOG(LogFastGeoStreaming, Log, TEXT("------------------------------------------------------------------------"));
+			UE_LOG(LogFastGeoStreaming, Log, TEXT("- FastGeoStreaming Debug Mode: Transformation on %d selected actors "), SelectedActors.Num());
 
-				FastGeo::FTransformationStats Stats;
-				Stats.TotalActorCount = SelectedActors.Num();
-				for (AActor* Actor : SelectedActors)
+			FastGeo::FTransformationStats Stats;
+			Stats.TotalActorCount = SelectedActors.Num();
+			for (AActor* Actor : SelectedActors)
+			{
+				bool bIsActorFullyTransformable;
+				TArray<UActorComponent*> TransformableComponents;
+				FFastGeoTransformResult ActorTransformResult = CanTransformActor(Actor, bIsActorFullyTransformable, TransformableComponents);
+				if (ActorTransformResult.GetResult() == EFastGeoTransform::Allow)
 				{
-					bool bIsActorFullyTransformable;
-					TArray<UActorComponent*> TransformableComponents;
-					FFastGeoTransformResult ActorTransformResult = CanTransformActor(Actor, bIsActorFullyTransformable, TransformableComponents);
-					if (ActorTransformResult.GetResult() == EFastGeoTransform::Allow)
-					{
-						Stats.FullyTransformableActorCount += bIsActorFullyTransformable ? 1 : 0;
-						Stats.PartiallyTransformableActorCount += !bIsActorFullyTransformable ? 1 : 0;
-						Stats.TransformedComponentCount += TransformableComponents.Num();
-					}
-					Stats.TotalComponentCount += Actor->GetComponents().Num();
+					Stats.FullyTransformableActorCount += bIsActorFullyTransformable ? 1 : 0;
+					Stats.PartiallyTransformableActorCount += !bIsActorFullyTransformable ? 1 : 0;
+					Stats.TransformedComponentCount += TransformableComponents.Num();
 				}
-
-				Stats.DumpStats(TEXT("  - "));
-				UE_LOG(LogFastGeoStreaming, Log, TEXT("------------------------------------------------------------------------"));
+				Stats.TotalComponentCount += Actor->GetComponents().Num();
 			}
+
+			Stats.DumpStats(TEXT("  - "));
+			UE_LOG(LogFastGeoStreaming, Log, TEXT("------------------------------------------------------------------------"));
 		}
 	}
 }
@@ -539,6 +534,11 @@ FFastGeoTransformResult UFastGeoWorldPartitionRuntimeCellTransformer::CanTransfo
 		return AllowedActorClassResult;
 	}
 
+	if (!InActor->GetIsSpatiallyLoaded())
+	{
+		return { EFastGeoTransform::Reject, [InActor] { return FString::Printf(TEXT("Actor %s is non-spatially loaded"), *InActor->GetName()); } };
+	}
+
 	if (InActor->GetIsReplicated())
 	{
 		return { EFastGeoTransform::Reject, [InActor] { return FString::Printf(TEXT("Actor %s is replicated"), *InActor->GetName()); } };
@@ -751,7 +751,7 @@ void UFastGeoWorldPartitionRuntimeCellTransformer::PostEditChangeProperty(FPrope
 {
 	const FName PropertyName = InPropertyChangedEvent.GetPropertyName();
 	if ((PropertyName == GET_MEMBER_NAME_CHECKED(UFastGeoWorldPartitionRuntimeCellTransformer, bDebugMode) || 
-		 PropertyName == GET_MEMBER_NAME_CHECKED(UFastGeoWorldPartitionRuntimeCellTransformer, bDebugModeOnEditorSelection)) && !FastGeo::GPackageWasDirty)
+		 PropertyName == GET_MEMBER_NAME_CHECKED(UFastGeoWorldPartitionRuntimeCellTransformer, bDebugModeOnSelection)) && !FastGeo::GPackageWasDirty)
 	{
 		GetPackage()->ClearDirtyFlag();
 	}
diff --git a/Engine/Plugins/Experimental/FastGeoStreaming/Source/FastGeoStreaming/Private/FastGeoWorldPartitionRuntimeCellTransformer.h b/Engine/Plugins/Experimental/FastGeoStreaming/Source/FastGeoStreaming/Private/FastGeoWorldPartitionRuntimeCellTransformer.h
--- a/Engine/Plugins/Experimental/FastGeoStreaming/Source/FastGeoStreaming/Private/FastGeoWorldPartitionRuntimeCellTransformer.h
+++ b/Engine/Plugins/Experimental/FastGeoStreaming/Source/FastGeoStreaming/Private/FastGeoWorldPartitionRuntimeCellTransformer.h
@@ -79,9 +79,9 @@ private:
 	UPROPERTY(EditAnywhere, Transient, SkipSerialization, Category = FastGeo)
 	bool bDebugMode;
 
-	/** When enabled, Editor will log the reason why selected actors/components are not transformable */
+	/** When enabled, will log the reason why selected actors/components are not transformable. Works in both PIE and Editor. */
 	UPROPERTY(EditAnywhere, Transient, SkipSerialization, Category = FastGeo)
-	bool bDebugModeOnEditorSelection;
+	bool bDebugModeOnSelection;
 
 	/** Allowed actor classes (recursive) to convert to FastGeo */
 	UPROPERTY(EditAnywhere, Category = FastGeo)
