FastGeo: Exposed FastGeoCellTransformer to allow licensees to override default implementation of IsActorTransformable, IsComponentTransformable, IsFullyTransformedActorDeletable.

diff --git a/Engine/Plugins/Experimental/FastGeoStreaming/Source/FastGeoStreaming/Private/FastGeoWorldPartitionRuntimeCellTransformer.cpp b/Engine/Plugins/Experimental/FastGeoStreaming/Source/FastGeoStreaming/Private/FastGeoWorldPartitionRuntimeCellTransformer.cpp
--- a/Engine/Plugins/Experimental/FastGeoStreaming/Source/FastGeoStreaming/Private/FastGeoWorldPartitionRuntimeCellTransformer.cpp
+++ b/Engine/Plugins/Experimental/FastGeoStreaming/Source/FastGeoStreaming/Private/FastGeoWorldPartitionRuntimeCellTransformer.cpp
@@ -534,6 +534,12 @@ FFastGeoTransformResult UFastGeoWorldPartitionRuntimeCellTransformer::CanTransfo
 		return AllowedActorClassResult;
 	}
 
+	FString Reason;
+	if (!IsActorTransformable(InActor, Reason))
+	{
+		return { EFastGeoTransform::Reject, [InActor, &Reason] { return FString::Printf(TEXT("Actor %s [%s]"), *InActor->GetName(), *Reason); } };
+	}
+
 	if (!InActor->GetIsSpatiallyLoaded())
 	{
 		return { EFastGeoTransform::Reject, [InActor] { return FString::Printf(TEXT("Actor %s is non-spatially loaded"), *InActor->GetName()); } };
@@ -590,6 +596,13 @@ FFastGeoTransformResult UFastGeoWorldPartitionRuntimeCellTransformer::CanTransfo
 	// If actor contains only discardable or transformable components, we can actually get rid of it
 	bOutIsActorFullyTransformable = TransformResults[EnumToIndex(EFastGeoTransform::Reject)].IsEmpty();
 
+	Reason.Reset();
+	if (!IsFullyTransformedActorDeletable(InActor, Reason))
+	{
+		bOutIsActorFullyTransformable = false;
+		UE_CLOG(FFastGeoTransformResult::ShouldReport, LogFastGeoStreaming, Log, TEXT("  * Can't delete fully transformed actor %s [%s]"), *InActor->GetName(), *Reason);
+	}
+
 	// Can't convert partially a BP actors
 	// Rerun CS will be called in PIE when registering the component and
 	// also called when registering components during cook/save of the level.
@@ -648,6 +661,12 @@ FFastGeoTransformResult UFastGeoWorldPartitionRuntimeCellTransformer::CanTransfo
 		return AllowedComponentClassResult;
 	}
 
+	FString Reason;
+	if (!IsComponentTransformable(InComponent, Reason))
+	{
+		return { EFastGeoTransform::Reject, [InComponent, &Reason] { return FString::Printf(TEXT("Component %s [%s]"), *FastGeo::GetComponentShortName(InComponent), *Reason); } };
+	}
+
 	if (InComponent->IsEditorOnly())
 	{
 		return { EFastGeoTransform::Discard, [InComponent] { return FString::Printf(TEXT("Component %s is editor-only"), *FastGeo::GetComponentShortName(InComponent)); } };
diff --git a/Engine/Plugins/Experimental/FastGeoStreaming/Source/FastGeoStreaming/Private/FastGeoWorldPartitionRuntimeCellTransformer.h b/Engine/Plugins/Experimental/FastGeoStreaming/Source/FastGeoStreaming/Private/FastGeoWorldPartitionRuntimeCellTransformer.h
deleted file mode 100644
diff --git a/Engine/Plugins/Experimental/FastGeoStreaming/Source/FastGeoStreaming/Public/FastGeoWorldPartitionRuntimeCellTransformer.h b/Engine/Plugins/Experimental/FastGeoStreaming/Source/FastGeoStreaming/Public/FastGeoWorldPartitionRuntimeCellTransformer.h
new file mode 100644
--- /dev/null
+++ b/Engine/Plugins/Experimental/FastGeoStreaming/Source/FastGeoStreaming/Public/FastGeoWorldPartitionRuntimeCellTransformer.h
@@ -0,0 +1,141 @@
+// Copyright Epic Games, Inc. All Rights Reserved.
+
+#pragma once
+
+#include "WorldPartition/WorldPartitionRuntimeCellTransformer.h"
+#if WITH_EDITOR
+#include "Templates/SubclassOf.h"
+#include "HAL/IConsoleManager.h"
+#endif 
+#include "FastGeoWorldPartitionRuntimeCellTransformer.generated.h"
+
+class UActorComponent;
+
+enum class EFastGeoTransform : uint32
+{
+	Allow,		//!< Actor or component can be transformed.
+	Reject,		//!< Actor or component can't be transformed.
+	Discard,	//!< Actor or component is not relevant and can be fully discarded without impact on the game.
+	MAX
+};
+
+inline constexpr uint32 EnumToIndex(EFastGeoTransform Value)
+{
+	return static_cast<uint32>(Value);
+}
+
+/** Used to conditionally report the error result when not a success. */
+struct FFastGeoTransformResult
+{
+public:
+	FFastGeoTransformResult(EFastGeoTransform InTransformResult, const TCHAR* FailureReason = nullptr);
+	FFastGeoTransformResult(EFastGeoTransform InTransformResult, TFunctionRef<FString()> FailureReasonFunc);
+
+	EFastGeoTransform GetResult() const { return TransformResult; }
+	uint32 GetResultIndex() const { return EnumToIndex(TransformResult); }
+
+	static bool ShouldReport;
+
+private:
+	EFastGeoTransform TransformResult;
+};
+
+UCLASS()
+class FASTGEOSTREAMING_API UFastGeoWorldPartitionRuntimeCellTransformer : public UWorldPartitionRuntimeCellTransformer
+{
+	GENERATED_UCLASS_BODY()
+
+public:
+	//~ Begin UObject Interface.
+	virtual void Serialize(FArchive& Ar) override;
+	//~ End UObject Interface.
+
+#if WITH_EDITOR
+public:
+	//~ Begin UWorldPartitionRuntimeCellTransformer Interface.
+	virtual void Transform(ULevel* InLevel) override final;
+	//~ End UWorldPartitionRuntimeCellTransformer Interface.
+
+	//~ Begin UObject Interface.
+	virtual void BeginDestroy() override;
+	virtual void PreEditChange(FProperty* InPropertyAboutToChange) override;
+	virtual void PostEditChangeProperty(FPropertyChangedEvent& InPropertyChangedEvent) override;
+	//~ End UObject Interface.
+
+protected:
+	/** Whether the actor can be processed by the transformer. */
+	virtual bool IsActorTransformable(AActor* InActor, FString& OutReason) const { return true; }
+	/** Whether the component can be processed by the transformer. */
+	virtual bool IsComponentTransformable(UPrimitiveComponent* InComponent, FString& OutReason) const { return true; }
+	/** Whether the fully transformed actor can be deleted by the transformer. */
+	virtual bool IsFullyTransformedActorDeletable(AActor* InActor, FString& OutReason) const { return true; }
+
+	bool IsBlueprintActorWithLogic(AActor* InActor) const;
+	FFastGeoTransformResult CanTransformActor(AActor* InActor, bool& bOutIsActorFullyTransformable, TArray<UActorComponent*>& OutTransformableComponents) const;
+	FFastGeoTransformResult IsAllowedActorClass(AActor* InActor) const;
+	FFastGeoTransformResult CanTransformComponent(UPrimitiveComponent* InComponent) const;
+	FFastGeoTransformResult IsAllowedComponentClass(UActorComponent* InComponent) const;
+	bool CanRemoveActor(AActor* InActor, TSet<UActorComponent*>& InIgnoredComponents) const;
+	bool CanAlwaysIgnoreActor(AActor* InActor) const;
+	void OnSelectionChanged(UObject* Object);
+	bool IsDebugMode() const;
+#endif
+
+#if WITH_EDITORONLY_DATA
+protected:
+	/** When enabled, PIE will log the reason why actors/components are not transformable */
+	UPROPERTY(EditAnywhere, Transient, SkipSerialization, Category = FastGeo)
+	bool bDebugMode;
+
+	/** When enabled, will log the reason why selected actors/components are not transformable. Works in both PIE and Editor. */
+	UPROPERTY(EditAnywhere, Transient, SkipSerialization, Category = FastGeo)
+	bool bDebugModeOnSelection;
+
+	/** Allowed actor classes (recursive) to convert to FastGeo */
+	UPROPERTY(EditAnywhere, Category = FastGeo)
+	TArray<TSubclassOf<AActor>> AllowedActorClasses;
+
+	/** Allowed actor classes (exact) to convert to FastGeo */
+	UPROPERTY(EditAnywhere, Category = FastGeo)
+	TArray<TSubclassOf<AActor>> AllowedExactActorClasses;
+
+	/** Allowed component classes (recursive) to convert to FastGeo */
+	UPROPERTY(EditAnywhere, Category = FastGeo)
+	TArray<TSubclassOf<UActorComponent>> AllowedComponentClasses;
+
+	/** Allowed component classes (exact) to convert to FastGeo */
+	UPROPERTY(EditAnywhere, Category = FastGeo)
+	TArray<TSubclassOf<UActorComponent>> AllowedExactComponentClasses;
+
+	/** Disallowed actor classes (recursive) to convert to FastGeo */
+	UPROPERTY(EditAnywhere, Category = FastGeo)
+	TArray<TSubclassOf<AActor>> DisallowedActorClasses;
+
+	/** Disallowed actor classes (exact) to convert to FastGeo */
+	UPROPERTY(EditAnywhere, Category = FastGeo)
+	TArray<TSubclassOf<AActor>> DisallowedExactActorClasses;
+
+	/** Disallowed component classes (recursive) to convert to FastGeo */
+	UPROPERTY(EditAnywhere, Category = FastGeo)
+	TArray<TSubclassOf<UActorComponent>> DisallowedComponentClasses;
+
+	/** Disallowed component classes (exact) to convert to FastGeo */
+	UPROPERTY(EditAnywhere, Category = FastGeo)
+	TArray<TSubclassOf<UActorComponent>> DisallowedExactComponentClasses;
+
+	/** Component classes (recursive) to ignore when deciding to destroy a converted actor to FastGeo */
+	UPROPERTY(EditAnywhere, Category = FastGeo)
+	TArray<TSubclassOf<UActorComponent>> IgnoredRemainingComponentClasses;
+
+	/** Component classes (exact) to ignore when deciding to destroy a converted actor to FastGeo */
+	UPROPERTY(EditAnywhere, Category = FastGeo)
+	TArray<TSubclassOf<UActorComponent>> IgnoredRemainingExactComponentClasses;
+#endif
+
+#if WITH_EDITOR
+	static bool IsDebugModeEnabled;
+	static bool IsFastGeoEnabled;
+	static FAutoConsoleVariableRef CVarIsDebugModeEnabled;
+	static FAutoConsoleVariableRef CVarIsFastGeoEnabled;
+#endif
+};
\ No newline at end of file
