FastGeo: Shared implementation of scene proxy creation between FastGeoStaticMeshComponent and FastGeoInstancedStaticMeshComponent.

diff --git a/Engine/Plugins/Experimental/FastGeoStreaming/Source/FastGeoStreaming/Private/FastGeoInstancedStaticMeshComponent.cpp b/Engine/Plugins/Experimental/FastGeoStreaming/Source/FastGeoStreaming/Private/FastGeoInstancedStaticMeshComponent.cpp
--- a/Engine/Plugins/Experimental/FastGeoStreaming/Source/FastGeoStreaming/Private/FastGeoInstancedStaticMeshComponent.cpp
+++ b/Engine/Plugins/Experimental/FastGeoStreaming/Source/FastGeoStreaming/Private/FastGeoInstancedStaticMeshComponent.cpp
@@ -159,53 +159,11 @@ TSharedPtr<FInstanceDataSceneProxy, ESPMode::ThreadSafe>& FFastGeoInstancedStati
 	return DataProxy;
 }
 
-FPrimitiveSceneProxy* FFastGeoInstancedStaticMeshComponent::CreateSceneProxy(ESceneProxyCreationError* OutError)
+FPrimitiveSceneProxy* FFastGeoInstancedStaticMeshComponent::CreateSceneProxyImpl()
 {
 	check(GetWorld());
-	FSceneInterface* Scene = GetScene();
-	UStaticMesh* StaticMesh = GetStaticMesh();
-	check(Scene);
-	check(StaticMesh);
-	check(StaticMesh->GetRenderData());
-	check(StaticMesh->GetRenderData()->IsInitialized());
-	check(!StaticMesh->IsCompiling());
-
-	if (OutError)
-	{
-		*OutError = ESceneProxyCreationError::None;
-	}
-
-#if !WITH_EDITOR && UE_WITH_PSO_PRECACHING
-	const EPSOPrecachePriority PSOPrecachePriority = GetStaticMeshComponentBoostPSOPrecachePriority();
-	if (CheckPSOPrecachingAndBoostPriority(PSOPrecachePriority) && (GetPSOPrecacheProxyCreationStrategy() == EPSOPrecacheProxyCreationStrategy::DelayUntilPSOPrecached))
-	{
-		UE_LOG(LogFastGeoStreaming, Verbose, TEXT("Skipping CreateSceneProxy for FastGeoInstancedStaticMeshComponent of %s (component PSOs are still compiling)"), *GetOwnerComponentCluster()->GetName());
-		if (OutError)
-		{
-			*OutError = ESceneProxyCreationError::WaitingPSOs;
-		}
-		return nullptr;
-	}
-#endif
-
-	InitializeSceneProxyDescDynamicProperties();
-
-	check(SceneProxyDesc.GetScene());
-	check(SceneProxyDesc.GetScene() == Scene);
-	check(SceneProxyDesc.World == Scene->GetWorld());
-	check(SceneProxyDesc.FeatureLevel == Scene->GetFeatureLevel());
-	check(SceneProxyDesc.ComponentId == GetPrimitiveSceneId());
-
-	// Verify that the mesh is valid before using it.
-	const bool bMeshIsValid = PerInstanceSMData.Num() > 0 && StaticMesh && !StaticMesh->IsCompiling() && (StaticMesh->HasValidRenderData() || StaticMesh->HasValidNaniteData());
-	if (!bMeshIsValid)
-	{
-		if (OutError)
-		{
-			*OutError = ESceneProxyCreationError::InvalidMesh;
-		}
-		return nullptr;
-	}
+	check(SceneProxyDesc.Scene);
+	check(PerInstanceSMData.Num() > 0);
 
 	SceneProxyDesc.InstanceDataSceneProxy = BuildInstanceData();
 	Nanite::FMaterialAudit MaterialAudit;
diff --git a/Engine/Plugins/Experimental/FastGeoStreaming/Source/FastGeoStreaming/Private/FastGeoInstancedStaticMeshComponent.h b/Engine/Plugins/Experimental/FastGeoStreaming/Source/FastGeoStreaming/Private/FastGeoInstancedStaticMeshComponent.h
--- a/Engine/Plugins/Experimental/FastGeoStreaming/Source/FastGeoStreaming/Private/FastGeoInstancedStaticMeshComponent.h
+++ b/Engine/Plugins/Experimental/FastGeoStreaming/Source/FastGeoStreaming/Private/FastGeoInstancedStaticMeshComponent.h
@@ -46,12 +46,12 @@ protected:
 	virtual void ResetSceneProxyDescUnsupportedProperties() override;
 #endif
 	virtual void ApplyWorldTransform(const FTransform& InTransform) override;
-	virtual FPrimitiveSceneProxy* CreateSceneProxy(ESceneProxyCreationError* OutError = nullptr) override;
 	//~ End FFastGeoPrimitiveComponent interface
 
 	//~ Begin FFastGeoStaticMeshComponentBase interface
 	virtual FStaticMeshSceneProxyDesc& GetStaticMeshSceneProxyDesc() override { return SceneProxyDesc; }
 	virtual const FStaticMeshSceneProxyDesc& GetStaticMeshSceneProxyDesc() const override { return SceneProxyDesc; }
+	virtual FPrimitiveSceneProxy* CreateSceneProxyImpl() override;
 #if !WITH_EDITOR && UE_WITH_PSO_PRECACHING
 	virtual void CollectPSOPrecacheData(const FPSOPrecacheParams& BasePrecachePSOParams, FMaterialInterfacePSOPrecacheParamsList& OutParams) override;
 #endif
diff --git a/Engine/Plugins/Experimental/FastGeoStreaming/Source/FastGeoStreaming/Private/FastGeoSkinnedMeshComponent.cpp b/Engine/Plugins/Experimental/FastGeoStreaming/Source/FastGeoStreaming/Private/FastGeoSkinnedMeshComponent.cpp
--- a/Engine/Plugins/Experimental/FastGeoStreaming/Source/FastGeoStreaming/Private/FastGeoSkinnedMeshComponent.cpp
+++ b/Engine/Plugins/Experimental/FastGeoStreaming/Source/FastGeoStreaming/Private/FastGeoSkinnedMeshComponent.cpp
@@ -129,7 +129,7 @@ FPrimitiveSceneProxy* FFastGeoSkinnedMeshComponentBase::CreateSceneProxy(ESceneP
 
 	USkinnedAsset* SkinnedAsset = GetSkinnedAsset();
 	check(SkinnedAsset);
-	check(!SkinnedAsset->IsCompiling()); // Todo: Wait with creation of proxy till stopped compiling.
+	check(!SkinnedAsset->IsCompiling());
 	
 	if (OutError)
 	{
diff --git a/Engine/Plugins/Experimental/FastGeoStreaming/Source/FastGeoStreaming/Private/FastGeoStaticMeshComponent.cpp b/Engine/Plugins/Experimental/FastGeoStreaming/Source/FastGeoStreaming/Private/FastGeoStaticMeshComponent.cpp
--- a/Engine/Plugins/Experimental/FastGeoStreaming/Source/FastGeoStreaming/Private/FastGeoStaticMeshComponent.cpp
+++ b/Engine/Plugins/Experimental/FastGeoStreaming/Source/FastGeoStreaming/Private/FastGeoStaticMeshComponent.cpp
@@ -214,7 +214,7 @@ FPrimitiveSceneProxy* FFastGeoStaticMeshComponentBase::CreateSceneProxy(EScenePr
 	check(StaticMesh);
 	check(StaticMesh->GetRenderData());
 	check(StaticMesh->GetRenderData()->IsInitialized());
-	check(!StaticMesh->IsCompiling()); // Todo: Wait with creation of proxy till stopped compiling.
+	check(!StaticMesh->IsCompiling());
 
 	if (OutError)
 	{
@@ -225,7 +225,7 @@ FPrimitiveSceneProxy* FFastGeoStaticMeshComponentBase::CreateSceneProxy(EScenePr
 	const EPSOPrecachePriority PSOPrecachePriority = GetStaticMeshComponentBoostPSOPrecachePriority();
 	if (CheckPSOPrecachingAndBoostPriority(PSOPrecachePriority) && (GetPSOPrecacheProxyCreationStrategy() == EPSOPrecacheProxyCreationStrategy::DelayUntilPSOPrecached))
 	{
-		UE_LOG(LogFastGeoStreaming, Verbose, TEXT("Skipping CreateSceneProxy for FastGeoStaticMeshComponent of %s (component PSOs are still compiling)"), *GetOwnerComponentCluster()->GetName());
+		UE_LOG(LogFastGeoStreaming, Verbose, TEXT("Skipping CreateSceneProxy for Component of %s (component PSOs are still compiling)"), *GetOwnerComponentCluster()->GetName());
 		if (OutError)
 		{
 			*OutError = ESceneProxyCreationError::WaitingPSOs;
@@ -243,6 +243,26 @@ FPrimitiveSceneProxy* FFastGeoStaticMeshComponentBase::CreateSceneProxy(EScenePr
 	check(SceneProxyDesc.FeatureLevel == Scene->GetFeatureLevel());
 	check(SceneProxyDesc.ComponentId == GetPrimitiveSceneId());
 
+	// Verify that the mesh is valid before using it.
+	const bool bIsMeshValid = StaticMesh->HasValidRenderData() || StaticMesh->HasValidNaniteData();
+	if (!bIsMeshValid)
+	{
+		if (OutError)
+		{
+			*OutError = ESceneProxyCreationError::InvalidMesh;
+		}
+		return nullptr;
+	}
+
+	return CreateSceneProxyImpl();
+}
+
+FPrimitiveSceneProxy* FFastGeoStaticMeshComponentBase::CreateSceneProxyImpl()
+{
+	check(GetWorld());
+	const FStaticMeshSceneProxyDesc& SceneProxyDesc = GetStaticMeshSceneProxyDesc();
+	check(SceneProxyDesc.Scene);
+
 	Nanite::FMaterialAudit MaterialAudit;
 	if (SceneProxyDesc.ShouldCreateNaniteProxy(&MaterialAudit))
 	{
diff --git a/Engine/Plugins/Experimental/FastGeoStreaming/Source/FastGeoStreaming/Private/FastGeoStaticMeshComponent.h b/Engine/Plugins/Experimental/FastGeoStreaming/Source/FastGeoStreaming/Private/FastGeoStaticMeshComponent.h
--- a/Engine/Plugins/Experimental/FastGeoStreaming/Source/FastGeoStreaming/Private/FastGeoStaticMeshComponent.h
+++ b/Engine/Plugins/Experimental/FastGeoStreaming/Source/FastGeoStreaming/Private/FastGeoStaticMeshComponent.h
@@ -87,6 +87,7 @@ protected:
 
 	virtual FStaticMeshSceneProxyDesc& GetStaticMeshSceneProxyDesc() = 0;
 	virtual const FStaticMeshSceneProxyDesc& GetStaticMeshSceneProxyDesc() const = 0;
+	virtual FPrimitiveSceneProxy* CreateSceneProxyImpl();
 	UMaterialInterface* GetMaterial(int32 MaterialIndex, bool bDoingNaniteMaterialAudit) const;
 	bool ShouldExportAsObstacle(const UNavCollisionBase& InNavCollision) const;
 #if !WITH_EDITOR && UE_WITH_PSO_PRECACHING
