Fixed CIS - Static Analysis following 41900638

diff --git a/Engine/Plugins/Experimental/FastGeoStreaming/Source/FastGeoStreaming/Private/FastGeoWorldPartitionRuntimeCellTransformer.cpp b/Engine/Plugins/Experimental/FastGeoStreaming/Source/FastGeoStreaming/Private/FastGeoWorldPartitionRuntimeCellTransformer.cpp
--- a/Engine/Plugins/Experimental/FastGeoStreaming/Source/FastGeoStreaming/Private/FastGeoWorldPartitionRuntimeCellTransformer.cpp
+++ b/Engine/Plugins/Experimental/FastGeoStreaming/Source/FastGeoStreaming/Private/FastGeoWorldPartitionRuntimeCellTransformer.cpp
@@ -537,7 +537,7 @@ FFastGeoTransformResult UFastGeoWorldPartitionRuntimeCellTransformer::CanTransfo
 	FString Reason;
 	if (!IsActorTransformable(InActor, Reason))
 	{
-		return { EFastGeoTransform::Reject, [InActor, &Reason] { return FString::Printf(TEXT("Actor %s [%s]"), *InActor->GetName(), *Reason); } };
+		return { EFastGeoTransform::Reject, [InActor, Reason] { return FString::Printf(TEXT("Actor %s [%s]"), *InActor->GetName(), *Reason); } };
 	}
 
 	if (!InActor->GetIsSpatiallyLoaded())
@@ -664,7 +664,7 @@ FFastGeoTransformResult UFastGeoWorldPartitionRuntimeCellTransformer::CanTransfo
 	FString Reason;
 	if (!IsComponentTransformable(InComponent, Reason))
 	{
-		return { EFastGeoTransform::Reject, [InComponent, &Reason] { return FString::Printf(TEXT("Component %s [%s]"), *FastGeo::GetComponentShortName(InComponent), *Reason); } };
+		return { EFastGeoTransform::Reject, [InComponent, Reason] { return FString::Printf(TEXT("Component %s [%s]"), *FastGeo::GetComponentShortName(InComponent), *Reason); } };
 	}
 
 	if (InComponent->IsEditorOnly())
