FastGeo - Removed deferred cleanup of skeletal mesh objects

diff --git a/Engine/Plugins/Experimental/FastGeoStreaming/Source/FastGeoStreaming/Private/FastGeoSkinnedMeshComponent.cpp b/Engine/Plugins/Experimental/FastGeoStreaming/Source/FastGeoStreaming/Private/FastGeoSkinnedMeshComponent.cpp
--- a/Engine/Plugins/Experimental/FastGeoStreaming/Source/FastGeoStreaming/Private/FastGeoSkinnedMeshComponent.cpp
+++ b/Engine/Plugins/Experimental/FastGeoStreaming/Source/FastGeoStreaming/Private/FastGeoSkinnedMeshComponent.cpp
@@ -176,21 +176,8 @@ void FFastGeoSkinnedMeshComponentBase::DestroyMeshObject()
 {
 	if (MeshObject)
 	{
-		if (IsInGameThread() || IsInParallelGameThread())
-		{
-			MeshObject->ReleaseResources();
-			BeginCleanup(MeshObject);
-		}
-		else
-		{
-			// MeshObject->ReleaseResources() needs to be called from the game thread
-			AsyncTask(ENamedThreads::GameThread, [MeshObject=MeshObject]()
-			{
-				MeshObject->ReleaseResources();
-				BeginCleanup(MeshObject);				
-			});
-		}
-		
+		MeshObject->ReleaseResources();
+		BeginCleanup(MeshObject);
 		MeshObject = nullptr;
 	}
 }
