FastGeo - Don't automatically exclude BP actors if their blueprint contains macro graphs
* Macro graphs could be used in the construction script, whose presence doesn't automatically exclude a BP actor from FastGeo

diff --git a/Engine/Plugins/Experimental/FastGeoStreaming/Source/FastGeoStreaming/Private/FastGeoWorldPartitionRuntimeCellTransformer.cpp b/Engine/Plugins/Experimental/FastGeoStreaming/Source/FastGeoStreaming/Private/FastGeoWorldPartitionRuntimeCellTransformer.cpp
--- a/Engine/Plugins/Experimental/FastGeoStreaming/Source/FastGeoStreaming/Private/FastGeoWorldPartitionRuntimeCellTransformer.cpp
+++ b/Engine/Plugins/Experimental/FastGeoStreaming/Source/FastGeoStreaming/Private/FastGeoWorldPartitionRuntimeCellTransformer.cpp
@@ -440,7 +440,7 @@ bool UFastGeoWorldPartitionRuntimeCellTransformer::IsBlueprintActorWithLogic(AAc
 	}
 
 	// Check if no extra functions, other than the user construction script (only AActor and subclasses of AActor have)
-	if ((Blueprint->FunctionGraphs.Num() > 1) || (Blueprint->MacroGraphs.Num() > 0))
+	if (Blueprint->FunctionGraphs.Num() > 1)
 	{
 		return true;
 	}
