FastGeo - FastGeoContainer now has a fence to properly wait for render state destruction

diff --git a/Engine/Plugins/Experimental/FastGeoStreaming/Source/FastGeoStreaming/Private/FastGeoContainer.cpp b/Engine/Plugins/Experimental/FastGeoStreaming/Source/FastGeoStreaming/Private/FastGeoContainer.cpp
--- a/Engine/Plugins/Experimental/FastGeoStreaming/Source/FastGeoStreaming/Private/FastGeoContainer.cpp
+++ b/Engine/Plugins/Experimental/FastGeoStreaming/Source/FastGeoStreaming/Private/FastGeoContainer.cpp
@@ -397,22 +397,7 @@ void UFastGeoContainer::OnDestroyRenderState_Concurrent()
 		}
 
 		PendingDestroy.RenderState.NumProcessed = PendingDestroy.RenderState.NumToProcess;
-
-#if DO_CHECK
-		bDestroyRenderCommandCompleted = !Context.HasPendingWork();
-#endif
-	}
-
-#if DO_CHECK
-	if (!bDestroyRenderCommandCompleted)
-	{
-		ENQUEUE_RENDER_COMMAND(UFastGeoContainerDestroyed)(
-			[this](FRHICommandListImmediate& RHICmdList)
-			{
-				bDestroyRenderCommandCompleted = true;
-			});
 	}
-#endif
 }
 
 void UFastGeoContainer::OnCreateRenderStateEnd_GameThread()
@@ -446,12 +431,16 @@ void UFastGeoContainer::OnDestroyRenderStateEnd_GameThread()
 	}
 }
 
-void UFastGeoContainer::FinishDestroy()
+void UFastGeoContainer::BeginDestroy()
 {
-#if DO_CHECK
-	check(bDestroyRenderCommandCompleted);
-#endif
-	Super::FinishDestroy();
+	Super::BeginDestroy();
+	DestroyFence.BeginFence();
+}
+
+bool UFastGeoContainer::IsReadyForFinishDestroy()
+{
+	bool bResult = Super::IsReadyForFinishDestroy() && DestroyFence.IsFenceComplete();
+	return bResult;
 }
 
 #if WITH_EDITOR
diff --git a/Engine/Plugins/Experimental/FastGeoStreaming/Source/FastGeoStreaming/Private/FastGeoContainer.h b/Engine/Plugins/Experimental/FastGeoStreaming/Source/FastGeoStreaming/Private/FastGeoContainer.h
--- a/Engine/Plugins/Experimental/FastGeoStreaming/Source/FastGeoStreaming/Private/FastGeoContainer.h
+++ b/Engine/Plugins/Experimental/FastGeoStreaming/Source/FastGeoStreaming/Private/FastGeoContainer.h
@@ -74,7 +74,8 @@ public:
 	FFastGeoComponentCluster* GetComponentCluster(uint32 InComponentClusterTypeID, int32 InComponentClusterIndex);
 
 	//~ Begin UObject Interface
-	virtual void FinishDestroy() override;
+	virtual void BeginDestroy() override;
+	virtual bool IsReadyForFinishDestroy() override;
 	virtual void Serialize(FArchive& Ar) override;
 	virtual void PostLoad() override;
 #if WITH_EDITOR
@@ -248,9 +249,7 @@ private:
 
 	bool bIsRegistered = false;
 
-#if DO_CHECK
-	TAtomic<bool> bDestroyRenderCommandCompleted = true;
-#endif
+	FRenderCommandFence DestroyFence;
 
 	friend class FFastGeoGatherFastGeoContainerAssetRefsArchive;
 };
\ No newline at end of file
