FastGeo - Improved FastGeo cell transformer
* Disallow transform of SkeletalMeshComponent / InstancedSkinnedMeshComponent if they are animated (animations are not unsupported at the moment)
* Disallow transform of SkeletalMeshComponent / InstancedSkinnedMeshComponent that have collision or are navigation relevant.
* Disallow transform of HISMC if they are non-nanite and the mesh has more than one LOD (FastGeo can only do ISMC, so we'd lose the Hierarchical logic)
* Don't exclude hidden actors automatically - rely on the component's filtering logic (hidden components with collision can be transformed)
* Fixed issue where the bIsHidden flag was always reset to false
* Rewrote component filtering logic to reject unsupported component classes earlier in the process.

diff --git a/Engine/Plugins/Experimental/FastGeoStreaming/Source/FastGeoStreaming/Private/FastGeoElementType.h b/Engine/Plugins/Experimental/FastGeoStreaming/Source/FastGeoStreaming/Private/FastGeoElementType.h
--- a/Engine/Plugins/Experimental/FastGeoStreaming/Source/FastGeoStreaming/Private/FastGeoElementType.h
+++ b/Engine/Plugins/Experimental/FastGeoStreaming/Source/FastGeoStreaming/Private/FastGeoElementType.h
@@ -7,16 +7,18 @@
 struct FFastGeoElementType
 {
 public:
+	static const FFastGeoElementType Invalid;
+
 	explicit FFastGeoElementType(const FFastGeoElementType* Parent = nullptr) : ID(++NextUniqueID), ParentType(Parent) {}
 
 	bool operator==(const FFastGeoElementType& Other) const
 	{
-		return ID == Other.ID || (ParentType != nullptr && *ParentType == Other);
+		return IsSameTypeID(Other.ID);
 	}
 
 	bool IsA(const FFastGeoElementType& Other) const
 	{
-		return (ID == Other.ID) || (ParentType && ParentType->IsA(Other));
+		return IsSameTypeID(Other.ID) || (ParentType && ParentType->IsA(Other));
 	}
 
 	/** Returns true if element has the same type ID than the specified type */
@@ -25,7 +27,15 @@ public:
 		return ID == InID;
 	}
 
+	bool IsValid() const
+	{
+		return !IsSameTypeID(Invalid.ID);
+	}
+
 private:
+	enum class EInitMode { Invalid };
+	explicit FFastGeoElementType(EInitMode) : ID(0xFFFFFFFF), ParentType(nullptr) {}
+
 	static uint32 NextUniqueID;
 	uint32 ID;
 	const FFastGeoElementType* ParentType;
diff --git a/Engine/Plugins/Experimental/FastGeoStreaming/Source/FastGeoStreaming/Private/FastGeoPrimitiveComponent.cpp b/Engine/Plugins/Experimental/FastGeoStreaming/Source/FastGeoStreaming/Private/FastGeoPrimitiveComponent.cpp
--- a/Engine/Plugins/Experimental/FastGeoStreaming/Source/FastGeoStreaming/Private/FastGeoPrimitiveComponent.cpp
+++ b/Engine/Plugins/Experimental/FastGeoStreaming/Source/FastGeoStreaming/Private/FastGeoPrimitiveComponent.cpp
@@ -207,7 +207,6 @@ void FFastGeoPrimitiveComponent::ResetSceneProxyDescUnsupportedProperties()
 	SceneProxyDesc.bIndividuallySelected = false;
 	SceneProxyDesc.bShouldRenderSelected = false;
 	SceneProxyDesc.bWantsEditorEffects = false;
-	SceneProxyDesc.bIsHidden = false;
 	SceneProxyDesc.bIsHiddenEd = false;
 	SceneProxyDesc.bIsOwnerEditorOnly = false;
 	SceneProxyDesc.bIsOwnedByFoliage = false;
diff --git a/Engine/Plugins/Experimental/FastGeoStreaming/Source/FastGeoStreaming/Private/FastGeoWorldPartitionRuntimeCellTransformer.cpp b/Engine/Plugins/Experimental/FastGeoStreaming/Source/FastGeoStreaming/Private/FastGeoWorldPartitionRuntimeCellTransformer.cpp
--- a/Engine/Plugins/Experimental/FastGeoStreaming/Source/FastGeoStreaming/Private/FastGeoWorldPartitionRuntimeCellTransformer.cpp
+++ b/Engine/Plugins/Experimental/FastGeoStreaming/Source/FastGeoStreaming/Private/FastGeoWorldPartitionRuntimeCellTransformer.cpp
@@ -13,6 +13,7 @@
 #include "Modules/ModuleManager.h"
 #include "Physics/Experimental/PhysScene_Chaos.h"
 #include "Streaming/ActorTextureStreamingBuildDataComponent.h"
+#include "Components/HierarchicalInstancedStaticMeshComponent.h"
 #include "Components/InstancedStaticMeshComponent.h"
 #include "Components/StaticMeshComponent.h"
 #include "Components/SkinnedMeshComponent.h"
@@ -29,6 +30,7 @@
 #include "LevelInstance/LevelInstanceEditorInstanceActor.h"
 #include "LevelInstance/LevelInstanceEditorPivotInterface.h"
 #include "ActorEditorUtils.h"
+#include "Animation/AnimInstance.h"
 #include "Selection.h"
 #include "FastGeoContainer.h"
 #include "FastGeoComponentCluster.h"
@@ -113,6 +115,28 @@ namespace FastGeo
 			}
 		}
 	};
+
+	static const TMap<TSubclassOf<UPrimitiveComponent>, FFastGeoElementType> FastGeoComponentTypeMapping =
+	{
+		{ UStaticMeshComponent::StaticClass(), FFastGeoStaticMeshComponent::Type },
+		{ UInstancedStaticMeshComponent::StaticClass(), FFastGeoInstancedStaticMeshComponent::Type },
+		{ USkinnedMeshComponent::StaticClass(), FFastGeoSkinnedMeshComponent::Type },
+		{ UInstancedSkinnedMeshComponent::StaticClass(), FFastGeoInstancedSkinnedMeshComponent::Type }
+	};
+		
+	FFastGeoElementType GetFastGeoComponentType(TSubclassOf<UPrimitiveComponent> InClass)
+	{
+		// Walk the component class hierarchy and look for a fast geo mapping
+		for (UClass* Class = InClass; Class; Class = Class->GetSuperClass())
+		{
+			if (const FFastGeoElementType* Found = FastGeoComponentTypeMapping.Find(Class))
+			{
+				return *Found;
+			}
+		}
+
+		return FFastGeoElementType::Invalid;
+	}
 }
 #endif
 
@@ -261,6 +285,13 @@ void UFastGeoWorldPartitionRuntimeCellTransformer::OnSelectionChanged(UObject* O
 			}
 
 			Stats.DumpStats(TEXT("  - "));
+
+			if (!FPhysScene::SupportsAsyncPhysicsStateCreation() ||
+				!FPhysScene::SupportsAsyncPhysicsStateDestruction())
+			{
+				UE_LOG(LogFastGeoStreaming, Warning, TEXT(" - NOTE: FastGeoStreaming requires 'p.Chaos.EnableAsyncInitBody' to be enabled."));
+			}
+
 			UE_LOG(LogFastGeoStreaming, Log, TEXT("------------------------------------------------------------------------"));
 		}
 	}
@@ -287,7 +318,7 @@ void UFastGeoWorldPartitionRuntimeCellTransformer::Transform(ULevel* InLevel)
 	if (!FPhysScene::SupportsAsyncPhysicsStateCreation() || 
 		!FPhysScene::SupportsAsyncPhysicsStateDestruction())
 	{
-		UE_LOG(LogFastGeoStreaming, Warning, TEXT("FastGeoStreaming Cell Transformer requires 'p.Chaos.EnableAsyncInitBody' to be enabled."));
+		UE_LOG(LogFastGeoStreaming, Error, TEXT("FastGeoStreaming Cell Transformer requires 'p.Chaos.EnableAsyncInitBody' to be enabled."));
 		return;
 	}
 
@@ -322,29 +353,6 @@ void UFastGeoWorldPartitionRuntimeCellTransformer::Transform(ULevel* InLevel)
 		}
 	}
 
-	const TMap<TSubclassOf<UPrimitiveComponent>, FFastGeoElementType> FastGeoComponentTypeMapping =
-	{
-		{ UStaticMeshComponent::StaticClass(), FFastGeoStaticMeshComponent::Type },
-		{ UInstancedStaticMeshComponent::StaticClass(), FFastGeoInstancedStaticMeshComponent::Type },
-		{ USkinnedMeshComponent::StaticClass(), FFastGeoSkinnedMeshComponent::Type },
-		{ UInstancedSkinnedMeshComponent::StaticClass(), FFastGeoInstancedSkinnedMeshComponent::Type }
-	};
-
-	// Walk the component class hierarchy and look for a fast geo mapping
-	auto GetFastGeoComponentType = [&FastGeoComponentTypeMapping](TSubclassOf<UPrimitiveComponent> InClass) -> FFastGeoElementType
-	{
-		for (UClass* Class = InClass; Class; Class = Class->GetSuperClass())
-		{
-			if (const FFastGeoElementType* Found = FastGeoComponentTypeMapping.Find(Class))
-			{
-				return *Found;
-			}
-		}
-
-		unimplemented(); // Should never be reached, missing type handling
-		return FFastGeoComponent::Type;
-	};
-
 	if (!ActorComponents.IsEmpty())
 	{
 		const FString CellName = InLevel->GetWorldPartitionRuntimeCell() ? Cast<UObject>(InLevel->GetWorldPartitionRuntimeCell())->GetName() : TEXT("Cell");
@@ -369,7 +377,9 @@ void UFastGeoWorldPartitionRuntimeCellTransformer::Transform(ULevel* InLevel)
 			check(!Components.IsEmpty());
 			for (UActorComponent* Component : Components)
 			{
-				FFastGeoElementType FastGeoComponentType = GetFastGeoComponentType(Component->GetClass());
+				FFastGeoElementType FastGeoComponentType = FastGeo::GetFastGeoComponentType(Component->GetClass());
+				check(FastGeoComponentType.IsValid());
+
 				FFastGeoComponent& FastGeoComponent = CurrentComponentCluster->AddComponent(FastGeoComponentType);
 				FastGeoComponent.InitializeFromComponent(Component);
 
@@ -555,11 +565,6 @@ FFastGeoTransformResult UFastGeoWorldPartitionRuntimeCellTransformer::CanTransfo
 		return { EFastGeoTransform::Reject, [InActor] { return FString::Printf(TEXT("Actor %s RootComponent Mobility is not Static"), *InActor->GetName()); } };
 	}
 
-	if (InActor->IsHidden())
-	{
-		return { EFastGeoTransform::Reject, [InActor] { return FString::Printf(TEXT("Actor %s is hidden"), *InActor->GetName()); } };
-	}
-
 	if (InActor->IsEditorOnly())
 	{
 		return { EFastGeoTransform::Reject, [InActor] { return FString::Printf(TEXT("Actor %s is editor-only"), *InActor->GetName()); } };
@@ -617,6 +622,12 @@ FFastGeoTransformResult UFastGeoWorldPartitionRuntimeCellTransformer::CanTransfo
 
 FFastGeoTransformResult UFastGeoWorldPartitionRuntimeCellTransformer::IsAllowedComponentClass(UActorComponent* InComponent) const
 {
+	FFastGeoElementType FastGeoComponentType = FastGeo::GetFastGeoComponentType(InComponent->GetClass());
+	if (!FastGeoComponentType.IsValid())
+	{
+		return { EFastGeoTransform::Reject, [InComponent] { return FString::Printf(TEXT("Component %s class is unsupported (%s)"), *FastGeo::GetComponentShortName(InComponent), *InComponent->GetClass()->GetName()); } };
+	}
+
 	UClass* ComponentClass = InComponent->GetClass();
 	for (TSubclassOf<UActorComponent> DisallowedComponentClass : DisallowedComponentClasses)
 	{
@@ -682,7 +693,8 @@ FFastGeoTransformResult UFastGeoWorldPartitionRuntimeCellTransformer::CanTransfo
 		return { EFastGeoTransform::Reject, [InComponent] { return FString::Printf(TEXT("Component %s Mobility is not Static"), *FastGeo::GetComponentShortName(InComponent)); } };
 	}
 
-	bool bIsComponentVisible = InComponent->IsVisible();
+	bool bIsComponentVisible = InComponent->IsVisible() && !InComponent->GetOwner()->IsHidden();
+	bool bShouldAddToRenderScene = bIsComponentVisible || InComponent->bCastHiddenShadow || InComponent->bAffectIndirectLightingWhileHidden || InComponent->bRayTracingFarField;
 
 	if (UStaticMeshComponent* StaticMeshComponent = Cast<UStaticMeshComponent>(InComponent))
 	{
@@ -703,26 +715,68 @@ FFastGeoTransformResult UFastGeoWorldPartitionRuntimeCellTransformer::CanTransfo
 		}
 	}
 
+	if (UHierarchicalInstancedStaticMeshComponent* HierarchicalInstancedStaticMeshComponent = Cast<UHierarchicalInstancedStaticMeshComponent>(InComponent))
+	{
+		// FastGeo doesn't really support HISMC. These components get converted to ISMC.
+		// However, we can afford to convert nanite HISMC as all the LODing logic is performed by Nanite.
+		// We also allow the transformation of HISMC which are using a mesh with a single LOD - so in effect it's handled as an ISMC.
+		if (!HierarchicalInstancedStaticMeshComponent->GetStaticMesh()->IsNaniteEnabled() || HierarchicalInstancedStaticMeshComponent->IsForceDisableNanite())
+		{
+			if (HierarchicalInstancedStaticMeshComponent->GetStaticMesh()->GetNumLODs() > 1)
+			{
+				return { EFastGeoTransform::Reject, [InComponent] { return FString::Printf(TEXT("Hierarchical instanced static mesh component %s has multiple LODs"), *FastGeo::GetComponentShortName(InComponent)); } };
+			}
+		}
+	}
+
 	if (USkinnedMeshComponent* SkinnedMeshComponent = Cast<USkinnedMeshComponent>(InComponent))
 	{
 		if (!SkinnedMeshComponent->GetSkinnedAsset())
 		{
-			return { EFastGeoTransform::Discard, [InComponent] { return FString::Printf(TEXT("Component %s has an invalid skinned asset"), *FastGeo::GetComponentShortName(InComponent)); } };
+			return { EFastGeoTransform::Discard, [InComponent] { return FString::Printf(TEXT("Skinned mesh component %s has an invalid skinned asset"), *FastGeo::GetComponentShortName(InComponent)); } };
 		}
 
 		if (SkinnedMeshComponent->LeaderPoseComponent.IsValid())
 		{
-			return { EFastGeoTransform::Reject, [InComponent] { return FString::Printf(TEXT("Skinned mesh Component %s has a leader pose component"), *FastGeo::GetComponentShortName(InComponent)); } };
+			return { EFastGeoTransform::Reject, [InComponent] { return FString::Printf(TEXT("Skinned mesh component %s has a leader pose component"), *FastGeo::GetComponentShortName(InComponent)); } };
 		}
 
-		bIsComponentVisible &= !SkinnedMeshComponent->bHideSkin;
+		if (SkinnedMeshComponent->GetCollisionEnabled() != ECollisionEnabled::NoCollision)
+		{
+			return { EFastGeoTransform::Reject, [InComponent] { return FString::Printf(TEXT("Skinned mesh component %s has collisions enabled"), *FastGeo::GetComponentShortName(InComponent)); } };
+		}
+
+		if (SkinnedMeshComponent->IsNavigationRelevant())
+		{
+			return { EFastGeoTransform::Reject, [InComponent] { return FString::Printf(TEXT("Skinned mesh component %s is navigation relevant"), *FastGeo::GetComponentShortName(InComponent)); } };
+		}
+
+		bShouldAddToRenderScene &= !SkinnedMeshComponent->bHideSkin;
+	}
+
+	if (USkeletalMeshComponent* SkeletalMeshComponent = Cast<USkeletalMeshComponent>(InComponent))
+	{
+		if ((SkeletalMeshComponent->GetAnimationMode() == EAnimationMode::AnimationSingleNode && SkeletalMeshComponent->AnimationData.AnimToPlay != nullptr) ||
+			(SkeletalMeshComponent->GetAnimationMode() == EAnimationMode::AnimationBlueprint && SkeletalMeshComponent->AnimClass != nullptr) ||
+			(SkeletalMeshComponent->GetAnimationMode() == EAnimationMode::AnimationCustomMode))
+		{
+			return { EFastGeoTransform::Reject, [InComponent] { return FString::Printf(TEXT("Skeletal mesh component %s is animated"), *FastGeo::GetComponentShortName(InComponent)); } };
+		}
 	}
 
 	if (UInstancedSkinnedMeshComponent* InstancedSkinnedMeshComponent = Cast<UInstancedSkinnedMeshComponent>(InComponent))
 	{
+		for (const FAnimBankItem& AnimBankItem : InstancedSkinnedMeshComponent->GetAnimBankItems())
+		{
+			if (AnimBankItem.BankAsset != nullptr)
+			{
+				return { EFastGeoTransform::Reject, [InComponent] { return FString::Printf(TEXT("Instanced skinned mesh component %s is animated"), *FastGeo::GetComponentShortName(InComponent)); } };
+			}
+		}
+
 		if (InstancedSkinnedMeshComponent->GetInstanceCount() == 0)
 		{
-			return { EFastGeoTransform::Discard, [InComponent] { return FString::Printf(TEXT("Component %s has no instances"), *FastGeo::GetComponentShortName(InComponent)); } };
+			return { EFastGeoTransform::Discard, [InComponent] { return FString::Printf(TEXT("Instanced skinned mesh component %s has no instances"), *FastGeo::GetComponentShortName(InComponent)); } };
 		}
 	}
 
@@ -734,8 +788,8 @@ FFastGeoTransformResult UFastGeoWorldPartitionRuntimeCellTransformer::CanTransfo
 		return { EFastGeoTransform::Reject, [InComponent] { return FString::Printf(TEXT("Component %s has collision enabled but doesn't allow asynchronous physics state creation/destruction"), *FastGeo::GetComponentShortName(InComponent)); } };
 	}
 
-	// Disallow transform if collision is disabled and component is not visible
-	if (!bIsCollisionEnabled && !bIsComponentVisible)
+	// Disallow transform if collision is disabled and component doesn't need to be added to the render scene
+	if (!bIsCollisionEnabled && !bShouldAddToRenderScene)
 	{
 		return { EFastGeoTransform::Discard, [InComponent] { return FString::Printf(TEXT("Component %s has no collision and is not visible"), *FastGeo::GetComponentShortName(InComponent)); } };
 	}
diff --git a/Engine/Plugins/Experimental/FastGeoStreaming/Source/FastGeoStreaming/Private/IFastGeoElement.cpp b/Engine/Plugins/Experimental/FastGeoStreaming/Source/FastGeoStreaming/Private/IFastGeoElement.cpp
--- a/Engine/Plugins/Experimental/FastGeoStreaming/Source/FastGeoStreaming/Private/IFastGeoElement.cpp
+++ b/Engine/Plugins/Experimental/FastGeoStreaming/Source/FastGeoStreaming/Private/IFastGeoElement.cpp
@@ -3,4 +3,6 @@
 #include "IFastGeoElement.h"
 
 uint32 FFastGeoElementType::NextUniqueID = 0;
-const FFastGeoElementType IFastGeoElement::Type;
\ No newline at end of file
+const FFastGeoElementType IFastGeoElement::Type;
+
+const FFastGeoElementType FFastGeoElementType::Invalid = FFastGeoElementType(FFastGeoElementType::EInitMode::Invalid);
\ No newline at end of file
