FastGeo - Moved hardcoded builtin classes out of the user editable arrays
* Made these builtin array properties "config" so that the lists of classes is defined in "data" rather than code

diff --git a/Engine/Plugins/Experimental/FastGeoStreaming/Config/FastGeoStreaming.ini b/Engine/Plugins/Experimental/FastGeoStreaming/Config/FastGeoStreaming.ini
new file mode 100644
--- /dev/null
+++ b/Engine/Plugins/Experimental/FastGeoStreaming/Config/FastGeoStreaming.ini
@@ -0,0 +1,17 @@
+[/Script/FastGeoStreaming.FastGeoWorldPartitionRuntimeCellTransformer]
++BuiltinAllowedActorClasses=/Script/CoreUObject.Class'/Script/Engine.PackedLevelActor'
++BuiltinAllowedActorClasses=/Script/CoreUObject.Class'/Script/Engine.StaticMeshActor'
++BuiltinAllowedActorClasses=/Script/CoreUObject.Class'/Script/Engine.WorldPartitionHLOD'
++BuiltinAllowedActorClasses=/Script/CoreUObject.Class'/Script/Engine.WorldPartitionAutoInstancedActor'
++BuiltinDisallowedActorClasses=/Script/CoreUObject.Class'/Script/Water.WaterBody'
++BuiltinDisallowedActorClasses=/Script/CoreUObject.Class'/Script/Engine.Info'
++BuiltinAllowedComponentClasses=/Script/CoreUObject.Class'/Script/Engine.StaticMeshComponent'
++BuiltinDisallowedComponentClasses=/Script/CoreUObject.Class'/Script/Foliage.InteractiveFoliageComponent'
++BuiltinDisallowedComponentClasses=/Script/CoreUObject.Class'/Script/Landscape.LandscapeNaniteComponent'
++BuiltinDisallowedComponentClasses=/Script/CoreUObject.Class'/Script/Landscape.LandscapeMeshProxyComponent'
++BuiltinDisallowedComponentClasses=/Script/CoreUObject.Class'/Script/Water.WaterBodyMeshComponent'
++BuiltinDisallowedComponentClasses=/Script/CoreUObject.Class'/Script/Engine.SplineMeshComponent'
++BuiltinIgnoredRemainingComponentClasses=/Script/CoreUObject.Class'/Script/Engine.LevelInstanceComponent'
++BuiltinIgnoredRemainingComponentClasses=/Script/CoreUObject.Class'/Script/Engine.ActorTextureStreamingBuildDataComponent'
++BuiltinIgnoredRemainingExactComponentClasses=/Script/CoreUObject.Class'/Script/Engine.SceneComponent'
+
diff --git a/Engine/Plugins/Experimental/FastGeoStreaming/Source/FastGeoStreaming/Private/FastGeoWorldPartitionRuntimeCellTransformer.cpp b/Engine/Plugins/Experimental/FastGeoStreaming/Source/FastGeoStreaming/Private/FastGeoWorldPartitionRuntimeCellTransformer.cpp
--- a/Engine/Plugins/Experimental/FastGeoStreaming/Source/FastGeoStreaming/Private/FastGeoWorldPartitionRuntimeCellTransformer.cpp
+++ b/Engine/Plugins/Experimental/FastGeoStreaming/Source/FastGeoStreaming/Private/FastGeoWorldPartitionRuntimeCellTransformer.cpp
@@ -44,9 +44,6 @@
 #endif
 
 #if WITH_EDITORONLY_DATA
-#include "Engine/StaticMeshActor.h"
-#include "LevelInstance/LevelInstanceComponent.h"
-#include "PackedLevelActor/PackedLevelActor.h"
 #include "WorldPartition/HLOD/HLODActor.h"
 #endif
 
@@ -78,18 +75,6 @@ namespace FastGeo
 		return *Builder;
 	}
 
-	struct FClassHelper
-	{
-		template <typename T>
-		static void AddClass(const TCHAR* Name, TArray<TSubclassOf<T>>& OutClassArray)
-		{
-			if (UClass* Class = FindObject<UClass>(nullptr, Name, true); ensure(Class))
-			{
-				OutClassArray.Add(Class);
-			}
-		}
-	};
-
 	struct FTransformationStats
 	{
 		int32 TotalActorCount = 0;
@@ -146,35 +131,8 @@ UFastGeoWorldPartitionRuntimeCellTransformer::UFastGeoWorldPartitionRuntimeCellT
 #if WITH_EDITOR
 	if (GIsEditor && !IsTemplate())
 	{
-		USelection::SelectionChangedEvent.AddUObject(this, &UFastGeoWorldPartitionRuntimeCellTransformer::OnSelectionChanged);	}
-#endif
-
-#if WITH_EDITORONLY_DATA
-	// Allowed actor/component classes
-	AllowedActorClasses.Add(APackedLevelActor::StaticClass());
-	AllowedActorClasses.Add(AStaticMeshActor::StaticClass());
-	AllowedActorClasses.Add(AWorldPartitionHLOD::StaticClass());
-	AllowedActorClasses.Add(AWorldPartitionAutoInstancedActor::StaticClass());
-	AllowedComponentClasses.Add(UStaticMeshComponent::StaticClass());
-
-	// Disallowed actor/component classes
-	FastGeo::FClassHelper::AddClass(TEXT("/Script/Foliage.InteractiveFoliageComponent"), DisallowedComponentClasses);
-	FastGeo::FClassHelper::AddClass(TEXT("/Script/Landscape.LandscapeNaniteComponent"), DisallowedComponentClasses);
-	FastGeo::FClassHelper::AddClass(TEXT("/Script/Landscape.LandscapeMeshProxyComponent"), DisallowedComponentClasses);
-	if (FModuleManager::Get().LoadModule(TEXT("Water")))
-	{
-		FastGeo::FClassHelper::AddClass(TEXT("/Script/Water.WaterBodyMeshComponent"), DisallowedComponentClasses);
-		FastGeo::FClassHelper::AddClass(TEXT("/Script/Water.WaterBody"), DisallowedActorClasses);
-	}
-	DisallowedActorClasses.Add(AInfo::StaticClass());
-	DisallowedComponentClasses.Add(USplineMeshComponent::StaticClass());
-
-	// Ignored component classes
-	IgnoredRemainingComponentClasses.Add(ULevelInstanceComponent::StaticClass());
-	IgnoredRemainingComponentClasses.Add(UActorTextureStreamingBuildDataComponent::StaticClass());
-
-	// Ignored remaining exact component classes
-	IgnoredRemainingExactComponentClasses.Add(USceneComponent::StaticClass());
+		USelection::SelectionChangedEvent.AddUObject(this, &UFastGeoWorldPartitionRuntimeCellTransformer::OnSelectionChanged);
+	}
 #endif
 }
 
@@ -495,6 +453,14 @@ FFastGeoTransformResult UFastGeoWorldPartitionRuntimeCellTransformer::IsAllowedA
 		}
 	}
 
+	for (TSubclassOf<AActor> DisallowedActorClass : BuiltinDisallowedActorClasses)
+	{
+		if (ActorClass->IsChildOf(DisallowedActorClass))
+		{
+			return { EFastGeoTransform::Reject, [InActor, DisallowedActorClass] { return FString::Printf(TEXT("Actor %s class is child of a built-in disallowed class (%s)"), *InActor->GetName(), *DisallowedActorClass->GetName()); } };
+		}
+	}
+
 	for (TSubclassOf<AActor> DisallowedActorClass : DisallowedExactActorClasses)
 	{
 		if (ActorClass == DisallowedActorClass)
@@ -517,6 +483,14 @@ FFastGeoTransformResult UFastGeoWorldPartitionRuntimeCellTransformer::IsAllowedA
 		}
 	}
 
+	for (TSubclassOf<AActor> AllowedActorClass : BuiltinAllowedActorClasses)
+	{
+		if (ActorClass->IsChildOf(AllowedActorClass))
+		{
+			return EFastGeoTransform::Allow;
+		}
+	}
+
 	for (TSubclassOf<AActor> AllowedActorClass : AllowedExactActorClasses)
 	{
 		if (ActorClass->IsChildOf(AllowedActorClass))
@@ -637,6 +611,14 @@ FFastGeoTransformResult UFastGeoWorldPartitionRuntimeCellTransformer::IsAllowedC
 		}
 	}
 
+	for (TSubclassOf<UActorComponent> DisallowedComponentClass : BuiltinDisallowedComponentClasses)
+	{
+		if (ComponentClass->IsChildOf(DisallowedComponentClass))
+		{
+			return { EFastGeoTransform::Reject, [InComponent, DisallowedComponentClass] { return FString::Printf(TEXT("Component %s class is child of a built-in disallowed class (%s)"), *FastGeo::GetComponentShortName(InComponent), *DisallowedComponentClass->GetName()); } };
+		}
+	}
+
 	for (TSubclassOf<UActorComponent> DisallowedComponentClass : DisallowedExactComponentClasses)
 	{
 		if (ComponentClass == DisallowedComponentClass)
@@ -653,6 +635,14 @@ FFastGeoTransformResult UFastGeoWorldPartitionRuntimeCellTransformer::IsAllowedC
 		}
 	}
 
+	for (TSubclassOf<UActorComponent> AllowedComponentClass : BuiltinAllowedComponentClasses)
+	{
+		if (ComponentClass->IsChildOf(AllowedComponentClass))
+		{
+			return EFastGeoTransform::Allow;
+		}
+	}
+
 	for (TSubclassOf<UActorComponent> AllowedComponentClass : AllowedExactComponentClasses)
 	{
 		if (ComponentClass == AllowedComponentClass)
@@ -803,6 +793,14 @@ FFastGeoTransformResult UFastGeoWorldPartitionRuntimeCellTransformer::CanTransfo
 		}
 	}
 
+	for (const TSubclassOf<UActorComponent>& IgnoredComponentClass : BuiltinIgnoredRemainingComponentClasses)
+	{
+		if (ComponentClass->IsChildOf(IgnoredComponentClass))
+		{
+			return { EFastGeoTransform::Discard };
+		}
+	}
+
 	for (const TSubclassOf<UActorComponent>& IgnoredComponentClass : IgnoredRemainingExactComponentClasses)
 	{
 		if (ComponentClass == IgnoredComponentClass)
@@ -811,6 +809,14 @@ FFastGeoTransformResult UFastGeoWorldPartitionRuntimeCellTransformer::CanTransfo
 		}
 	}
 
+	for (const TSubclassOf<UActorComponent>& IgnoredComponentClass : BuiltinIgnoredRemainingExactComponentClasses)
+	{
+		if (ComponentClass == IgnoredComponentClass)
+		{
+			return { EFastGeoTransform::Discard };
+		}
+	}
+
 	return EFastGeoTransform::Allow;
 }
 
@@ -852,4 +858,4 @@ FFastGeoTransformResult::FFastGeoTransformResult(EFastGeoTransform InTransformRe
 
 #endif
 
-#undef LOCTEXT_NAMESPACE
\ No newline at end of file
+#undef LOCTEXT_NAMESPACE
diff --git a/Engine/Plugins/Experimental/FastGeoStreaming/Source/FastGeoStreaming/Public/FastGeoWorldPartitionRuntimeCellTransformer.h b/Engine/Plugins/Experimental/FastGeoStreaming/Source/FastGeoStreaming/Public/FastGeoWorldPartitionRuntimeCellTransformer.h
--- a/Engine/Plugins/Experimental/FastGeoStreaming/Source/FastGeoStreaming/Public/FastGeoWorldPartitionRuntimeCellTransformer.h
+++ b/Engine/Plugins/Experimental/FastGeoStreaming/Source/FastGeoStreaming/Public/FastGeoWorldPartitionRuntimeCellTransformer.h
@@ -40,7 +40,7 @@ private:
 	EFastGeoTransform TransformResult;
 };
 
-UCLASS()
+UCLASS(Config = FastGeoStreaming, DefaultConfig)
 class FASTGEOSTREAMING_API UFastGeoWorldPartitionRuntimeCellTransformer : public UWorldPartitionRuntimeCellTransformer
 {
 	GENERATED_UCLASS_BODY()
@@ -138,4 +138,24 @@ protected:
 	static FAutoConsoleVariableRef CVarIsDebugModeEnabled;
 	static FAutoConsoleVariableRef CVarIsFastGeoEnabled;
 #endif
+
+#if WITH_EDITORONLY_DATA
+	UPROPERTY(Config)
+	TArray<TSubclassOf<AActor>> BuiltinAllowedActorClasses;
+
+	UPROPERTY(Config)
+	TArray<TSubclassOf<AActor>> BuiltinDisallowedActorClasses;
+
+	UPROPERTY(Config)
+	TArray<TSubclassOf<UActorComponent>> BuiltinAllowedComponentClasses;
+
+	UPROPERTY(Config)
+	TArray<TSubclassOf<UActorComponent>> BuiltinDisallowedComponentClasses;
+
+	UPROPERTY(Config)
+	TArray<TSubclassOf<UActorComponent>> BuiltinIgnoredRemainingComponentClasses;
+
+	UPROPERTY(Config)
+	TArray<TSubclassOf<UActorComponent>> BuiltinIgnoredRemainingExactComponentClasses;
+#endif
 };
\ No newline at end of file
