FastGeo - Detect references to actors that are going to be converted to FastGeo
* Exclude actors that have:
-> Non FastGeo referencers
-> FastGeo referencers that are going to be only partially transformed

Refactored code to allow detection logic to run on the "Debug Mode on Selection" debug feature

diff --git a/Engine/Plugins/Experimental/FastGeoStreaming/Source/FastGeoStreaming/Private/FastGeoWorldPartitionRuntimeCellTransformer.cpp b/Engine/Plugins/Experimental/FastGeoStreaming/Source/FastGeoStreaming/Private/FastGeoWorldPartitionRuntimeCellTransformer.cpp
--- a/Engine/Plugins/Experimental/FastGeoStreaming/Source/FastGeoStreaming/Private/FastGeoWorldPartitionRuntimeCellTransformer.cpp
+++ b/Engine/Plugins/Experimental/FastGeoStreaming/Source/FastGeoStreaming/Private/FastGeoWorldPartitionRuntimeCellTransformer.cpp
@@ -226,21 +226,10 @@ void UFastGeoWorldPartitionRuntimeCellTransformer::OnSelectionChanged(UObject* O
 			UE_LOG(LogFastGeoStreaming, Log, TEXT("------------------------------------------------------------------------"));
 			UE_LOG(LogFastGeoStreaming, Log, TEXT("- FastGeoStreaming Debug Mode: Transformation on %d selected actors "), SelectedActors.Num());
 
-			FastGeo::FTransformationStats Stats;
-			Stats.TotalActorCount = SelectedActors.Num();
-			for (AActor* Actor : SelectedActors)
-			{
-				bool bIsActorFullyTransformable;
-				TArray<UActorComponent*> TransformableComponents;
-				FFastGeoTransformResult ActorTransformResult = CanTransformActor(Actor, bIsActorFullyTransformable, TransformableComponents);
-				if (ActorTransformResult.GetResult() == EFastGeoTransform::Allow)
-				{
-					Stats.FullyTransformableActorCount += bIsActorFullyTransformable ? 1 : 0;
-					Stats.PartiallyTransformableActorCount += !bIsActorFullyTransformable ? 1 : 0;
-					Stats.TransformedComponentCount += TransformableComponents.Num();
-				}
-				Stats.TotalComponentCount += Actor->GetComponents().Num();
-			}
+			FTransformationStats Stats;
+			TMap<AActor*, FTransformableActor> TransformableActors;
+
+			GatherTransformableActors(SelectedActors, SelectedActors[0]->GetLevel(), TransformableActors, Stats);
 
 			Stats.DumpStats(TEXT("  - "));
 
@@ -287,31 +276,12 @@ void UFastGeoWorldPartitionRuntimeCellTransformer::Transform(ULevel* InLevel)
 		UE_LOG(LogFastGeoStreaming, Log, TEXT("------------------------------------------------------------------------"));
 		UE_LOG(LogFastGeoStreaming, Log, TEXT("- FastGeoStreaming Debug Mode: Transforming Level '%s'"), *InLevel->GetPathName());
 	}
+	
+	FTransformationStats Stats;
+	TMap<AActor*, FTransformableActor> TransformableActors;
+	GatherTransformableActors(InLevel->Actors, InLevel, TransformableActors, Stats);
 
-	FastGeo::FTransformationStats Stats;
-	TMap<TObjectPtr<AActor>*, TPair<bool, TArray<UActorComponent*>>> ActorComponents;
-	for (TObjectPtr<AActor>& Actor : InLevel->Actors)
-	{
-		if (IsValid(Actor) && !CanAlwaysIgnoreActor(Actor))
-		{
-			Stats.TotalActorCount++;
-			Stats.TotalComponentCount += Actor->GetComponents().Num();
-
-			bool bIsActorFullyTransformable;
-			TArray<UActorComponent*> TransformableComponents;
-			FFastGeoTransformResult ActorTransformResult = CanTransformActor(Actor, bIsActorFullyTransformable, TransformableComponents);
-			if (ActorTransformResult.GetResult() == EFastGeoTransform::Allow)
-			{
-				Stats.FullyTransformableActorCount += bIsActorFullyTransformable ? 1 : 0;
-				Stats.PartiallyTransformableActorCount += !bIsActorFullyTransformable ? 1 : 0;
-				Stats.TransformedComponentCount += TransformableComponents.Num();
-
-				ActorComponents.FindOrAdd(&Actor, { bIsActorFullyTransformable, MoveTemp(TransformableComponents) });
-			}
-		}
-	}
-
-	if (!ActorComponents.IsEmpty())
+	if (!TransformableActors.IsEmpty())
 	{
 		const FString CellName = InLevel->GetWorldPartitionRuntimeCell() ? Cast<UObject>(InLevel->GetWorldPartitionRuntimeCell())->GetName() : TEXT("Cell");
 		UFastGeoContainer* FastGeo = NewObject<UFastGeoContainer>(InLevel, *FString::Printf(TEXT("FastGeoContainer_%s"), *CellName));
@@ -319,19 +289,22 @@ void UFastGeoWorldPartitionRuntimeCellTransformer::Transform(ULevel* InLevel)
 
 		TUniquePtr<FFastGeoComponentCluster> LevelComponentCluster(new FFastGeoComponentCluster(FastGeo, *FString::Printf(TEXT("FastGeoComponentCluster_%s"), *CellName)));
 
-		for (auto& [Actor, ActorInfo] : ActorComponents)
+		for (TMap<AActor*, FTransformableActor>::TConstIterator It = TransformableActors.CreateConstIterator(); It; ++It)
 		{
+			const TPair<AActor*, FTransformableActor>& Entry = *It;
+			AActor* Actor = Entry.Key;
+
 			FFastGeoComponentCluster* CurrentComponentCluster = LevelComponentCluster.Get();
 			TUniquePtr<FFastGeoHLOD> FastGeoHLOD;
-			if (AWorldPartitionHLOD* HLODActor = Cast<AWorldPartitionHLOD>(*Actor))
+			if (AWorldPartitionHLOD* HLODActor = Cast<AWorldPartitionHLOD>(Actor))
 			{
-				FastGeoHLOD.Reset(new FFastGeoHLOD(FastGeo, *FString::Printf(TEXT("FastGeoHLOD_%s"), *(*Actor)->GetName())));
+				FastGeoHLOD.Reset(new FFastGeoHLOD(FastGeo, *FString::Printf(TEXT("FastGeoHLOD_%s"), *Actor->GetName())));
 				FastGeoHLOD->SetSourceCellGuid(HLODActor->GetSourceCellGuid());
 				FastGeoHLOD->SetRequireWarmup(HLODActor->DoesRequireWarmup());
 				CurrentComponentCluster = FastGeoHLOD.Get();
 			}
 
-			const TArray<UActorComponent*>& Components = ActorInfo.Value;
+			const TArray<UActorComponent*>& Components = Entry.Value.TransformableComponents;
 			check(!Components.IsEmpty());
 			for (UActorComponent* Component : Components)
 			{
@@ -342,14 +315,14 @@ void UFastGeoWorldPartitionRuntimeCellTransformer::Transform(ULevel* InLevel)
 				FastGeoComponent.InitializeFromComponent(Component);
 
 				// Remove the component from the actor
-				(*Actor)->RemoveOwnedComponent(Component);
+				Actor->RemoveOwnedComponent(Component);
 				Component->MarkAsGarbage();
 			}
 
-			const bool bIsActorFullyTransformed = ActorInfo.Key;
+			const bool bIsActorFullyTransformed = Entry.Value.bIsActorFullyTransformable;
 			if (bIsActorFullyTransformed)
 			{
-				*Actor = nullptr;
+				InLevel->Actors[Entry.Value.ActorIndex] = nullptr;
 			}
 
 			if (FastGeoHLOD.IsValid())
@@ -379,6 +352,161 @@ void UFastGeoWorldPartitionRuntimeCellTransformer::Transform(ULevel* InLevel)
 	}
 }
 
+TMap<AActor*, TArray<AActor*>> UFastGeoWorldPartitionRuntimeCellTransformer::BuildActorsReferencesMap(const TArray<AActor*>& InActors)
+{
+	TRACE_CPUPROFILER_EVENT_SCOPE(UFastGeoWorldPartitionRuntimeCellTransformer::BuildActorsReferencesMap);
+
+	TMap<AActor*, TArray<AActor*>> ReferencedActors;
+	TSet<UObject*> VisitedObjects;
+	
+	// Visit all actors properties and look for references to other actors
+	for (AActor* ReferencingActor : InActors)
+	{
+		if (!IsValid(ReferencingActor))
+		{
+			continue;
+		}
+
+		if (CanAlwaysIgnoreActor(ReferencingActor))
+		{
+			continue;
+		}
+
+		VisitedObjects.Reset();
+		VisitedObjects.Add(ReferencingActor);
+
+		ReferencingActor->GetClass()->Visit(ReferencingActor, [&ReferencingActor, &ReferencedActors, &VisitedObjects](const FPropertyVisitorContext& Context) -> EPropertyVisitorControlFlow
+		{
+			const FPropertyVisitorPath& Path = Context.Path;
+			const FPropertyVisitorData& Data = Context.Data;
+			const FProperty* Property = Path.Top().Property;
+
+			// Step over editor only properties
+			if (Property->IsEditorOnlyProperty())
+			{
+				return EPropertyVisitorControlFlow::StepOver;
+			}
+
+			if (const FObjectPropertyBase* ObjectProperty = CastField<FObjectPropertyBase>(Property))
+			{
+				if (UObject* PropertyObject = ObjectProperty->GetObjectPropertyValue(Data.PropertyData))
+				{
+					bool bWasAlreadyInSet;
+					VisitedObjects.Add(PropertyObject, &bWasAlreadyInSet);
+
+					if (bWasAlreadyInSet)
+					{
+						return EPropertyVisitorControlFlow::StepOver;
+					}
+
+					AActor* ReferencedActor = Cast<AActor>(PropertyObject);
+					if (!ReferencedActor)
+					{
+						ReferencedActor = PropertyObject->GetTypedOuter<AActor>();
+					}
+
+					if (ReferencedActor && !ReferencedActor->HasAnyFlags(RF_ClassDefaultObject) && ReferencedActor != ReferencingActor)
+					{
+						ReferencedActors.FindOrAdd(ReferencedActor).Add(ReferencingActor);
+					}
+
+					// Constrain visitor to properties of objects that have ReferencingActor in their outer chain
+					if (!PropertyObject->IsIn(ReferencingActor))
+					{
+						return EPropertyVisitorControlFlow::StepOver;
+					}
+				}
+			}
+
+			return EPropertyVisitorControlFlow::StepInto;
+		}, FPropertyVisitorContext::EScope::ObjectRefs);
+	}
+
+	return ReferencedActors;
+}
+
+void UFastGeoWorldPartitionRuntimeCellTransformer::GatherTransformableActors(const TArray<AActor*>& InActors, const ULevel* InLevel, TMap<AActor*, FTransformableActor>& OutTransformableActors, FTransformationStats& OutStats)
+{
+	// Get transformation result for each actor
+	// This will retrieve the transformable components, and whether the actor is fully transformable (ie. if the actor can be deleted)
+	for (int32 ActorIndex = 0; ActorIndex < InActors.Num(); ++ActorIndex)
+	{
+		AActor* Actor = InActors[ActorIndex];
+		if (IsValid(Actor) && !CanAlwaysIgnoreActor(Actor))
+		{
+			OutStats.TotalActorCount++;
+			OutStats.TotalComponentCount += Actor->GetComponents().Num();
+
+			bool bIsActorFullyTransformable;
+			TArray<UActorComponent*> TransformableComponents;
+			FFastGeoTransformResult ActorTransformResult = CanTransformActor(Actor, bIsActorFullyTransformable, TransformableComponents);
+			if (ActorTransformResult.GetResult() == EFastGeoTransform::Allow)
+			{
+				OutTransformableActors.FindOrAdd(Actor, { ActorIndex, bIsActorFullyTransformable, MoveTemp(TransformableComponents) });
+			}
+		}
+	}
+
+	TMap<AActor*, TArray<AActor*>> ReferencedActors = BuildActorsReferencesMap(InLevel->Actors);
+	for (TMap<AActor*, FTransformableActor>::TIterator It = OutTransformableActors.CreateIterator(); It; ++It)
+	{
+		const TPair<AActor*, FTransformableActor>& Entry = *It;
+		const AActor* Actor = Entry.Key;
+
+		// Exclude actors that have:
+		//  * Non FastGeo referencers
+		//  * FastGeo referencers that are going to be only partially transformed
+		if (const TArray<AActor*>* Referencers = ReferencedActors.Find(Actor))
+		{
+			const AActor* Referencer = nullptr;
+			const FTransformableActor* TransformableReferencer = nullptr;
+			bool bHasNonFastGeoReferencer = Algo::AnyOf(*Referencers, [&Referencer, &TransformableReferencer, OutTransformableActors](const AActor* InReferencer)
+			{
+				Referencer = InReferencer;
+				TransformableReferencer = OutTransformableActors.Find(InReferencer);
+				return TransformableReferencer == nullptr || !TransformableReferencer->bIsActorFullyTransformable;
+			});
+
+			// If one of the referencer is not fully transformed to FastGeo
+			if (bHasNonFastGeoReferencer)
+			{
+				if (TransformableReferencer)
+				{
+					FFastGeoTransformResult(EFastGeoTransform::Reject, [Actor, TransformableReferencer, &InActors] { return FString::Printf(TEXT("Actor '%s' is referenced by a non-fully transformed actor ('%s')"), *Actor->GetName(), *InActors[TransformableReferencer->ActorIndex]->GetName()); });
+				}
+				else
+				{
+					FFastGeoTransformResult(EFastGeoTransform::Reject, [Actor, &Referencer] { return FString::Printf(TEXT("Actor '%s' is referenced by a non transformed actor ('%s')"), *Actor->GetName(), *Referencer->GetName()); });
+				}
+
+				It.RemoveCurrent();
+				continue;
+			}
+		}
+
+		OutStats.FullyTransformableActorCount += Entry.Value.bIsActorFullyTransformable ? 1 : 0;
+		OutStats.PartiallyTransformableActorCount += !Entry.Value.bIsActorFullyTransformable ? 1 : 0;
+		OutStats.TransformedComponentCount += Entry.Value.TransformableComponents.Num();
+	}
+}
+
+void UFastGeoWorldPartitionRuntimeCellTransformer::FTransformationStats::DumpStats(const TCHAR* InPrefixString)
+{
+	if (TotalActorCount)
+	{
+		const float FullyTransformedActorPercentage = TotalActorCount > 0 ? (100.f * FullyTransformableActorCount) / TotalActorCount : 0.f;
+		const float PartiallyTransformedActorPercentage = TotalActorCount > 0 ? (100.f * PartiallyTransformableActorCount) / TotalActorCount : 0.f;
+		const float TransformedComponentPercentage = TotalComponentCount > 0 ? (100.f * TransformedComponentCount) / TotalComponentCount : 0.f;
+		const int32 NonTransformableActorCount = FMath::Max(0, TotalActorCount - FullyTransformableActorCount - PartiallyTransformableActorCount);
+		const float NonTransformableActorPercentage = NonTransformableActorCount > 0 ? (100.f * NonTransformableActorCount) / TotalActorCount : 0.f;
+
+		UE_CLOG(FullyTransformableActorCount, LogFastGeoStreaming, Log, TEXT("%s Transformable Actors (Full)    = %d (%3.1f%%)"), InPrefixString, FullyTransformableActorCount, FullyTransformedActorPercentage);
+		UE_CLOG(PartiallyTransformableActorCount, LogFastGeoStreaming, Log, TEXT("%s Transformable Actors (Partial) = %d (%3.1f%%)"), InPrefixString, PartiallyTransformableActorCount, PartiallyTransformedActorPercentage);
+		UE_CLOG(TransformedComponentCount, LogFastGeoStreaming, Log, TEXT("%s Transformable Components       = %d (%3.1f%%)"), InPrefixString, TransformedComponentCount, TransformedComponentPercentage);
+		UE_CLOG(NonTransformableActorCount, LogFastGeoStreaming, Log, TEXT("%s Non-Transformable Actors       = %d (%3.1f%%)"), InPrefixString, NonTransformableActorCount, NonTransformableActorPercentage);
+	}
+}
+
 bool UFastGeoWorldPartitionRuntimeCellTransformer::IsBlueprintActorWithLogic(AActor* InActor) const
 {
 	const FName FN_UserConstructionScript(TEXT("UserConstructionScript"));
diff --git a/Engine/Plugins/Experimental/FastGeoStreaming/Source/FastGeoStreaming/Public/FastGeoWorldPartitionRuntimeCellTransformer.h b/Engine/Plugins/Experimental/FastGeoStreaming/Source/FastGeoStreaming/Public/FastGeoWorldPartitionRuntimeCellTransformer.h
--- a/Engine/Plugins/Experimental/FastGeoStreaming/Source/FastGeoStreaming/Public/FastGeoWorldPartitionRuntimeCellTransformer.h
+++ b/Engine/Plugins/Experimental/FastGeoStreaming/Source/FastGeoStreaming/Public/FastGeoWorldPartitionRuntimeCellTransformer.h
@@ -79,6 +79,27 @@ protected:
 	bool CanAlwaysIgnoreActor(AActor* InActor) const;
 	void OnSelectionChanged(UObject* Object);
 	bool IsDebugMode() const;
+
+	struct FTransformableActor
+	{
+		int32 ActorIndex = 0;
+		bool bIsActorFullyTransformable = false;
+		TArray<UActorComponent*> TransformableComponents;
+	};
+
+	struct FTransformationStats
+	{
+		int32 TotalActorCount = 0;
+		int32 TotalComponentCount = 0;
+		int32 FullyTransformableActorCount = 0;
+		int32 PartiallyTransformableActorCount = 0;
+		int32 TransformedComponentCount = 0;
+
+		void DumpStats(const TCHAR* InPrefixString);
+	};
+
+	TMap<AActor*, TArray<AActor*>> BuildActorsReferencesMap(const TArray<AActor*>& InActors);
+	void GatherTransformableActors(const TArray<AActor*>& InActors, const ULevel* InLevel, TMap<AActor*, FTransformableActor>& OutTransformableActors, FTransformationStats& OutStats);
 #endif
 
 #if WITH_EDITORONLY_DATA
