FastGeo - Fix issue where the FStaticMeshSceneProxyDesc::bIsHidden flag wasn't serialized and resulted in issues in cooked builds

diff --git a/Engine/Plugins/Experimental/FastGeoStreaming/Source/FastGeoStreaming/Private/FastGeoPrimitiveComponent.cpp b/Engine/Plugins/Experimental/FastGeoStreaming/Source/FastGeoStreaming/Private/FastGeoPrimitiveComponent.cpp
--- a/Engine/Plugins/Experimental/FastGeoStreaming/Source/FastGeoStreaming/Private/FastGeoPrimitiveComponent.cpp
+++ b/Engine/Plugins/Experimental/FastGeoStreaming/Source/FastGeoStreaming/Private/FastGeoPrimitiveComponent.cpp
@@ -123,6 +123,7 @@ void FFastGeoPrimitiveComponent::Serialize(FArchive& Ar)
 	FArchive_Serialize_BitfieldBool(Ar, SceneProxyDesc.bIsFirstPerson);
 	FArchive_Serialize_BitfieldBool(Ar, SceneProxyDesc.bIsFirstPersonWorldSpaceRepresentation);
 	FArchive_Serialize_BitfieldBool(Ar, SceneProxyDesc.bCollisionEnabled);
+	FArchive_Serialize_BitfieldBool(Ar, SceneProxyDesc.bIsHidden);
 	FArchive_Serialize_BitfieldBool(Ar, SceneProxyDesc.bSupportsWorldPositionOffsetVelocity);
 	FArchive_Serialize_BitfieldBool(Ar, SceneProxyDesc.bIsInstancedStaticMesh);
 	FArchive_Serialize_BitfieldBool(Ar, SceneProxyDesc.bHasStaticLighting);
