FastGeo: Moved FastGeoContainer PSO Precaching from PostLoad (GameThread) to be asynchronous.

diff --git a/Engine/Plugins/Experimental/FastGeoStreaming/Source/FastGeoStreaming/Private/FastGeoContainer.cpp b/Engine/Plugins/Experimental/FastGeoStreaming/Source/FastGeoStreaming/Private/FastGeoContainer.cpp
--- a/Engine/Plugins/Experimental/FastGeoStreaming/Source/FastGeoStreaming/Private/FastGeoContainer.cpp
+++ b/Engine/Plugins/Experimental/FastGeoStreaming/Source/FastGeoStreaming/Private/FastGeoContainer.cpp
@@ -20,6 +20,7 @@
 #include "Physics/Experimental/PhysScene_Chaos.h"
 #include "ProfilingDebugging/CountersTrace.h"
 #include "PSOPrecache.h"
+#include "Templates/Invoke.h"
 #include "WorldPartition/WorldPartitionHelpers.h"
 #include "WorldPartition/WorldPartition.h"
 
@@ -46,60 +47,130 @@ namespace FastGeo
 		GAsyncRenderStateTaskMaxNumComponentsToProcess,
 		TEXT("Maximum number of components to process (0 = no component limit)"),
 		ECVF_Default);
-}
 
-class FFastGeoAssetRemapArchive : public FArchiveProxy
-{
-public:
-	FFastGeoAssetRemapArchive(FArchive& InArchive, TArray<TObjectPtr<UObject>>& InUniqueAssetsArray)
-		: FArchiveProxy(InArchive)
-		, UniqueAssetsArray(InUniqueAssetsArray)
+	class FAssetRemapArchive : public FArchiveProxy
 	{
-		// For some unknown reason, copy constructor resets ArIsFilterEditorOnly flag copied from the input archive (see FArchiveState(const FArchiveState&))
-		ArIsFilterEditorOnly = InArchive.ArIsFilterEditorOnly;
-		for (int32 Index = 0, Num = UniqueAssetsArray.Num(); Index < Num; ++Index)
+	public:
+		FAssetRemapArchive(FArchive& InArchive, TArray<TObjectPtr<UObject>>& InUniqueAssetsArray)
+			: FArchiveProxy(InArchive)
+			, UniqueAssetsArray(InUniqueAssetsArray)
 		{
-			UniqueAssets.Add(UniqueAssetsArray[Index], Index);
+			// For some unknown reason, copy constructor resets ArIsFilterEditorOnly flag copied from the input archive (see FArchiveState(const FArchiveState&))
+			ArIsFilterEditorOnly = InArchive.ArIsFilterEditorOnly;
+			for (int32 Index = 0, Num = UniqueAssetsArray.Num(); Index < Num; ++Index)
+			{
+				UniqueAssets.Add(UniqueAssetsArray[Index], Index);
+			}
 		}
-	}
 
-	virtual FArchive& operator<<(UObject*& Obj) override
+		virtual FArchive& operator<<(UObject*& Obj) override
+		{
+			if (IsLoading())
+			{
+				int32 Index;
+				*this << Index;
+				Obj = UniqueAssetsArray.IsValidIndex(Index) ? UniqueAssetsArray[Index] : nullptr;
+			}
+			else if (IsSaving())
+			{
+				int32 Index = INDEX_NONE;
+				if (int32* ExistingIndex = Obj ? UniqueAssets.Find(Obj) : nullptr)
+				{
+					Index = *ExistingIndex;
+				}
+				*this << Index;
+			}
+			return *this;
+		}
+
+		virtual FArchive& operator<<(FObjectPtr& Obj) override
+		{
+			UObject* ObjPtr = Obj.Get();
+			FArchive& Result = operator<<(ObjPtr);
+			Obj = ObjPtr;
+			return Result;
+		}
+
+		virtual FArchive& operator<<(FLazyObjectPtr& LazyObjectPtr) override { unimplemented(); return *this; }
+		virtual FArchive& operator<<(FSoftObjectPtr& AssetPtr) override { unimplemented(); return *this; }
+		virtual FArchive& operator<<(FSoftObjectPath& AssetPtr) override { unimplemented(); return *this; }
+		virtual FArchive& operator<<(FWeakObjectPtr& Value) override { unimplemented(); return *this; }
+
+	private:
+		TMap<UObject*, int32> UniqueAssets;
+		TArray<TObjectPtr<UObject>>& UniqueAssetsArray;
+	};
+
+	class FAsyncTaskAssetReferenceManager : public FGCObject
 	{
-		if (IsLoading())
+	public:
+		TMap<FGuid, TArray<TObjectPtr<UObject>>> TasksAssets;
+
+		// Made an on-demand singleton rather than a static global, to avoid issues with FGCObject initialization
+		static FAsyncTaskAssetReferenceManager& Get()
 		{
-			int32 Index;
-			*this << Index;
-			Obj = UniqueAssetsArray.IsValidIndex(Index) ? UniqueAssetsArray[Index] : nullptr;
+			static FAsyncTaskAssetReferenceManager Manager;
+			return Manager;
 		}
-		else if (IsSaving())
+
+		void RegisterTaskAssets(const FGuid& TaskId, const TArray<TObjectPtr<UObject>>& Assets)
+		{
+			TasksAssets.Add(TaskId, Assets);
+		}
+
+		void UnregisterTask(const FGuid& TaskId)
 		{
-			int32 Index = INDEX_NONE;
-			if (int32* ExistingIndex = Obj ? UniqueAssets.Find(Obj) : nullptr)
+			TasksAssets.Remove(TaskId);
+		}
+
+		virtual void AddReferencedObjects(FReferenceCollector& Collector) override
+		{
+			for (auto& Pair : TasksAssets)
 			{
-				Index = *ExistingIndex;
+				Collector.AddReferencedObjects(Pair.Value);
 			}
-			*this << Index;
 		}
-		return *this;
-	}
 
-	virtual FArchive& operator<<(FObjectPtr& Obj) override
+		virtual FString GetReferencerName() const override
+		{
+			return TEXT("FAsyncTaskAssetReferenceManager");
+		}
+	};
+
+	struct FAsyncTaskWithAssetRefs : public TSharedFromThis<FAsyncTaskWithAssetRefs>
 	{
-		UObject* ObjPtr = Obj.Get();
-		FArchive& Result = operator<<(ObjPtr);
-		Obj = ObjPtr;
-		return Result;
-	}
+		explicit FAsyncTaskWithAssetRefs(const TArray<TObjectPtr<UObject>>& InAssets)
+			: TaskId(FGuid::NewGuid())
+			, Assets(InAssets)
+		{
+			FAsyncTaskAssetReferenceManager::Get().RegisterTaskAssets(TaskId, Assets);
+		}
 
-	virtual FArchive& operator<<(FLazyObjectPtr& LazyObjectPtr) override { unimplemented(); return *this; }
-	virtual FArchive& operator<<(FSoftObjectPtr& AssetPtr) override { unimplemented(); return *this; }
-	virtual FArchive& operator<<(FSoftObjectPath& AssetPtr) override { unimplemented(); return *this; }
-	virtual FArchive& operator<<(FWeakObjectPtr& Value) override { unimplemented(); return *this; }
+		~FAsyncTaskWithAssetRefs()
+		{
+			// Ensure unregistration happens on the game thread
+			UE::Tasks::Launch(TEXT("UnregisterFastGeoTask"), [MyTaskId = TaskId]()
+			{
+				FAsyncTaskAssetReferenceManager::Get().UnregisterTask(MyTaskId);
+			}, LowLevelTasks::ETaskPriority::Normal, UE::Tasks::EExtendedTaskPriority::GameThreadNormalPri);
+		}
 
-private:
-	TMap<UObject*, int32> UniqueAssets;
-	TArray<TObjectPtr<UObject>>& UniqueAssetsArray;
-};
+		template<typename TaskBodyType>
+		static UE::Tasks::TTask<TInvokeResult_T<TaskBodyType>> Launch(const TCHAR* TaskName, const TArray<TObjectPtr<UObject>>& AssetRefs, TaskBodyType&& Work, LowLevelTasks::ETaskPriority Priority)
+		{
+			TSharedRef<FAsyncTaskWithAssetRefs> TaskData = MakeShared<FAsyncTaskWithAssetRefs>(AssetRefs);
+
+			return UE::Tasks::Launch(TaskName, [TaskData, Work = MoveTemp(Work)]() mutable
+			{
+				Work();
+			}, Priority);
+		}
+
+	private:
+		FGuid TaskId;
+		TArray<TObjectPtr<UObject>> Assets;
+	};
+}
 
 void UFastGeoContainer::Register()
 {
@@ -346,13 +417,13 @@ void UFastGeoContainer::OnCreateRenderState_Concurrent()
 	const bool bTimeLimit = FastGeo::GAsyncRenderStateTaskTimeBudgetMS != 0;
 	const double TimeBudgetSeconds = FastGeo::GAsyncRenderStateTaskTimeBudgetMS / 1000.0;
 	UE::FTimeout Timeout = bTimeLimit ? UE::FTimeout(TimeBudgetSeconds) : UE::FTimeout::Never();
-	
+
 	TAtomic<int32> NextIndex = 0;
 	TAtomic<int32> NumProcessed = 0;
 
 	ParallelFor(NumThreads, [&](int32 Index)
 	{
-		TRACE_CPUPROFILER_EVENT_SCOPE("UFastGeoContainer::OnCreateRenderState_Concurrent");
+		TRACE_CPUPROFILER_EVENT_SCOPE(UFastGeoContainer::OnCreateRenderState_Concurrent);
 		FTaskTagScope Scope(ETaskTag::EParallelGameThread);
 
 		while(true)
@@ -408,7 +479,6 @@ void UFastGeoContainer::OnCreateRenderStateEnd_GameThread()
 	PendingCreate.RenderState.TotalNumProcessed += PendingCreate.RenderState.NumProcessed;
 	PendingCreate.RenderState.NumToProcess = 0;
 	PendingCreate.RenderState.NumProcessed = 0;
-	
 
 	if (PendingCreate.RenderState.TotalNumProcessed < PendingCreate.RenderState.ComponentsToProcess.Num())
 	{
@@ -440,6 +510,9 @@ void UFastGeoContainer::BeginDestroy()
 bool UFastGeoContainer::IsReadyForFinishDestroy()
 {
 	bool bResult = Super::IsReadyForFinishDestroy() && DestroyFence.IsFenceComplete();
+#if !WITH_EDITOR && UE_WITH_PSO_PRECACHING
+	bResult = bResult && PrecachePSOsTask.IsCompleted();
+#endif
 	return bResult;
 }
 
@@ -568,7 +641,7 @@ void UFastGeoContainer::Serialize(FArchive& Ar)
 {
 	Super::Serialize(Ar);
 
-	FFastGeoAssetRemapArchive AssetRemapAr(Ar, Assets);
+	FastGeo::FAssetRemapArchive AssetRemapAr(Ar, Assets);
 	SerializeComponentClusters(AssetRemapAr);
 
 #if !WITH_EDITOR
@@ -586,24 +659,39 @@ void UFastGeoContainer::SerializeComponentClusters(FArchive& Ar)
 	Ar << HLODs;
 }
 
-void UFastGeoContainer::PostLoad()
+void UFastGeoContainer::PrecachePSOs()
 {
-	TRACE_CPUPROFILER_EVENT_SCOPE(UFastGeoContainer::PostLoad);
-
-	Super::PostLoad();
-
-	UE_CLOG(!FPhysScene::SupportsAsyncPhysicsStateCreation() || !FPhysScene::SupportsAsyncPhysicsStateDestruction(), LogFastGeoStreaming, Error, TEXT("FastGeoStreaming requires 'p.Chaos.EnableAsyncInitBody' to be enabled."));
+	TRACE_CPUPROFILER_EVENT_SCOPE(UFastGeoContainer::PrecachePSOs_GameThread);
 
 #if !WITH_EDITOR && UE_WITH_PSO_PRECACHING
-	if (IsComponentPSOPrecachingEnabled())
+	if (bPrecachedPSOs)
 	{
-		ForEachComponentCluster([](FFastGeoComponentCluster& ComponentCluster)
+		return;
+	}
+	bPrecachedPSOs = true;
+
+	if (IsComponentPSOPrecachingEnabled() && (GetPSOPrecacheProxyCreationStrategy() != EPSOPrecacheProxyCreationStrategy::AlwaysCreate))
+	{
+		TArray<FFastGeoPrimitiveComponent*> Components;
+		ForEachComponentCluster([&Components](FFastGeoComponentCluster& ComponentCluster)
 		{
-			ComponentCluster.ForEachComponent<FFastGeoPrimitiveComponent>([](FFastGeoPrimitiveComponent& Component)
+			ComponentCluster.ForEachComponent<FFastGeoPrimitiveComponent>([&Components](FFastGeoPrimitiveComponent& Component)
 			{
-				Component.PrecachePSOs();
+				// Mark component so that IsPSOPrecaching() return true even if component's PSO task has not started
+				Component.MarkPrecachePSOsRequired();
+				Components.Add(&Component);
 			});
 		});
+
+		check(PrecachePSOsTask.IsCompleted());
+		PrecachePSOsTask = FastGeo::FAsyncTaskWithAssetRefs::Launch(TEXT("UFastGeoContainer::PrecachePSOs_Task"), Assets, [this, Components = MoveTemp(Components)]()
+		{
+			TRACE_CPUPROFILER_EVENT_SCOPE(UFastGeoContainer::PrecachePSOs_Task);
+			for (FFastGeoPrimitiveComponent* Component : Components)
+			{
+				Component->PrecachePSOs();
+			}
+		}, LowLevelTasks::ETaskPriority::BackgroundNormal);
 	}
 #endif
 }
diff --git a/Engine/Plugins/Experimental/FastGeoStreaming/Source/FastGeoStreaming/Private/FastGeoContainer.h b/Engine/Plugins/Experimental/FastGeoStreaming/Source/FastGeoStreaming/Private/FastGeoContainer.h
--- a/Engine/Plugins/Experimental/FastGeoStreaming/Source/FastGeoStreaming/Private/FastGeoContainer.h
+++ b/Engine/Plugins/Experimental/FastGeoStreaming/Source/FastGeoStreaming/Private/FastGeoContainer.h
@@ -69,6 +69,7 @@ public:
 	bool HasAnyPendingTasks() const;
 	bool HasAnyPendingCreateTasks() const;
 	bool HasAnyPendingDestroyTasks() const;
+	void PrecachePSOs();
 	UWorld* GetWorld() const;
 	ULevel* GetLevel() const;
 	FFastGeoComponentCluster* GetComponentCluster(uint32 InComponentClusterTypeID, int32 InComponentClusterIndex);
@@ -77,7 +78,6 @@ public:
 	virtual void BeginDestroy() override;
 	virtual bool IsReadyForFinishDestroy() override;
 	virtual void Serialize(FArchive& Ar) override;
-	virtual void PostLoad() override;
 #if WITH_EDITOR
 	virtual void PreSave(FObjectPreSaveContext ObjectSaveContext) override;
 #endif
@@ -243,6 +243,11 @@ private:
 	TSet<TObjectPtr<UObject>> EditorProxyObjects;
 #endif
 
+#if !WITH_EDITOR && UE_WITH_PSO_PRECACHING
+	bool bPrecachedPSOs = false;
+	UE::Tasks::TTask<void> PrecachePSOsTask;
+#endif
+
 	// Transient data
 	TArray<FFastGeoComponent*> CollisionComponents;
 	TArray<FNavigationElementHandle> NavigationElementHandles;
diff --git a/Engine/Plugins/Experimental/FastGeoStreaming/Source/FastGeoStreaming/Private/FastGeoPrimitiveComponent.cpp b/Engine/Plugins/Experimental/FastGeoStreaming/Source/FastGeoStreaming/Private/FastGeoPrimitiveComponent.cpp
--- a/Engine/Plugins/Experimental/FastGeoStreaming/Source/FastGeoStreaming/Private/FastGeoPrimitiveComponent.cpp
+++ b/Engine/Plugins/Experimental/FastGeoStreaming/Source/FastGeoStreaming/Private/FastGeoPrimitiveComponent.cpp
@@ -308,13 +308,13 @@ FPrimitiveSceneDesc FFastGeoPrimitiveComponent::BuildSceneDesc()
 
 void FFastGeoPrimitiveComponent::CreateRenderState(FRegisterComponentContext* Context)
 {
-	ProxyState = ProxyCreationState::Creating;
+	ProxyState = EProxyCreationState::Creating;
 	bRenderStateDirty = false;
 
 #if WITH_EDITOR
 	GetEditorProxy<EditorProxyType>()->NotifyRenderStateChanged();
 #endif
-
+	
 	FSceneInterface* Scene = GetScene();
 	check(Scene);
 
@@ -326,15 +326,15 @@ void FFastGeoPrimitiveComponent::CreateRenderState(FRegisterComponentContext* Co
 		FPrimitiveSceneDesc Desc = BuildSceneDesc();
 		Scene->AddPrimitive(&Desc);
 
-		ProxyState = ProxyCreationState::Created;
+		ProxyState = EProxyCreationState::Created;
 	}
 	else if (Error == ESceneProxyCreationError::WaitingPSOs)
 	{
-		ProxyState = ProxyCreationState::Delayed;
+		ProxyState = EProxyCreationState::Delayed;
 	}
 	else
 	{
-		ProxyState = ProxyCreationState::None;
+		ProxyState = EProxyCreationState::None;
 	}
 }
 
@@ -372,13 +372,13 @@ void FFastGeoPrimitiveComponent::DestroyRenderState(FFastGeoDestroyRenderStateCo
 {
 	if (GetSceneProxy())
 	{
-		check(ProxyState == ProxyCreationState::Created);
+		check(ProxyState == EProxyCreationState::Created);
 
 		FFastGeoDestroyRenderStateContext::DestroyProxy(Context, GetSceneProxy());
 
 		PrimitiveSceneData.SceneProxy = nullptr;
 
-		ProxyState = ProxyCreationState::Pending;
+		ProxyState = EProxyCreationState::Pending;
 
 		bRenderStateDirty = false;
 
@@ -484,7 +484,12 @@ void FFastGeoPrimitiveComponent::OnAsyncDestroyPhysicsState()
 
 bool FFastGeoPrimitiveComponent::IsRenderStateCreated() const
 {
-	return ProxyState == ProxyCreationState::Created;
+	return ProxyState == EProxyCreationState::Created;
+}
+
+bool FFastGeoPrimitiveComponent::IsRenderStateDelayed() const
+{
+	return ProxyState == EProxyCreationState::Delayed;
 }
 
 bool FFastGeoPrimitiveComponent::IsRenderStateDirty() const
@@ -523,7 +528,7 @@ FPrimitiveSceneProxy* FFastGeoPrimitiveComponent::GetSceneProxy() const
 
 void FFastGeoPrimitiveComponent::MarkRenderStateDirty()
 {
-	if ((IsRenderStateCreated() || ProxyState == ProxyCreationState::Delayed) && !IsRenderStateDirty())
+	if ((IsRenderStateCreated() || IsRenderStateDelayed()) && !IsRenderStateDirty())
 	{
 		bRenderStateDirty = true;
 
@@ -661,7 +666,8 @@ void FFastGeoPrimitiveComponent::SetupPrecachePSOParams(FPSOPrecacheParams& Para
 bool FFastGeoPrimitiveComponent::IsPSOPrecaching() const
 {
 #if !WITH_EDITOR && UE_WITH_PSO_PRECACHING
-	return LatestPSOPrecacheJobSetCompleted != LatestPSOPrecacheJobSet;
+	// Consider as precaching when marked as required to do PSOs precaching (even if task has not been launched yet) 
+	return bPSOPrecacheRequired || (LatestPSOPrecacheJobSetCompleted != LatestPSOPrecacheJobSet);
 #else
 	return false;
 #endif
@@ -681,42 +687,59 @@ bool FFastGeoPrimitiveComponent::CheckPSOPrecachingAndBoostPriority(EPSOPrecache
 {
 	bool bPrecacheStillRunning = IsPSOPrecaching();
 
-	ensure(!IsComponentPSOPrecachingEnabled() || bPSOPrecacheCalled);
+	ensure(!IsComponentPSOPrecachingEnabled() || bPSOPrecacheCalled || bPSOPrecacheRequired);
 	check(NewPSOPrecachePriority == EPSOPrecachePriority::High || NewPSOPrecachePriority == EPSOPrecachePriority::Highest);
 
 	if (bPrecacheStillRunning && PSOPrecacheRequestPriority < NewPSOPrecachePriority)
 	{
-		BoostPSOPriority(NewPSOPrecachePriority, MaterialPSOPrecacheRequestIDs);
+		// Only boost PSO priority if PSO task was started
+		if (LatestPSOPrecacheJobSetCompleted != LatestPSOPrecacheJobSet)
+		{
+			BoostPSOPriority(NewPSOPrecachePriority, MaterialPSOPrecacheRequestIDs);
+		}
 		PSOPrecacheRequestPriority = NewPSOPrecachePriority;
 	}
 	return bPrecacheStillRunning;
 }
 #endif
 
-void FFastGeoPrimitiveComponent::PrecachePSOs()
+void FFastGeoPrimitiveComponent::MarkPrecachePSOsRequired()
 {
 #if !WITH_EDITOR && UE_WITH_PSO_PRECACHING
-	TRACE_CPUPROFILER_EVENT_SCOPE(FFastGeoPrimitiveComponent::PrecachePSOs);
-
-	if (!FApp::CanEverRender() || !IsComponentPSOPrecachingEnabled())
+	if (bPSOPrecacheCalled || !FApp::CanEverRender() || !IsComponentPSOPrecachingEnabled())
 	{
 		return;
 	}
+	bPSOPrecacheRequired = true;
+	PSOPrecacheRequestPriority = EPSOPrecachePriority::Medium;
+#endif
+}
+
+void FFastGeoPrimitiveComponent::PrecachePSOs()
+{
+#if !WITH_EDITOR && UE_WITH_PSO_PRECACHING
+	TRACE_CPUPROFILER_EVENT_SCOPE(FFastGeoPrimitiveComponent::PrecachePSOs);
+	check(bPSOPrecacheRequired);
 
 	// Clear the current request data
 	MaterialPSOPrecacheRequestIDs.Empty();
-	PSOPrecacheRequestPriority = EPSOPrecachePriority::Medium;
 
 	// Collect the data from the derived classes
 	FPSOPrecacheParams PSOPrecacheParams;
 	SetupPrecachePSOParams(PSOPrecacheParams);
 	FMaterialInterfacePSOPrecacheParamsList PSOPrecacheDataArray;
 	CollectPSOPrecacheData(PSOPrecacheParams, PSOPrecacheDataArray);
+	// Set priority
+	for (FMaterialInterfacePSOPrecacheParams& Params : PSOPrecacheDataArray)
+	{
+		Params.Priority = PSOPrecacheRequestPriority;
+	}
 
 	FGraphEventArray GraphEvents;
 	PrecacheMaterialPSOs(PSOPrecacheDataArray, MaterialPSOPrecacheRequestIDs, GraphEvents);
 
 	RequestRecreateRenderStateWhenPSOPrecacheFinished(GraphEvents);
+	bPSOPrecacheRequired = false;
 #endif
 }
 
diff --git a/Engine/Plugins/Experimental/FastGeoStreaming/Source/FastGeoStreaming/Private/FastGeoPrimitiveComponent.h b/Engine/Plugins/Experimental/FastGeoStreaming/Source/FastGeoStreaming/Private/FastGeoPrimitiveComponent.h
--- a/Engine/Plugins/Experimental/FastGeoStreaming/Source/FastGeoStreaming/Private/FastGeoPrimitiveComponent.h
+++ b/Engine/Plugins/Experimental/FastGeoStreaming/Source/FastGeoStreaming/Private/FastGeoPrimitiveComponent.h
@@ -109,10 +109,12 @@ public:
 
 	virtual void DestroyRenderState(FFastGeoDestroyRenderStateContext* Context);
 	bool IsRenderStateCreated() const;
+	bool IsRenderStateDelayed() const;
 	bool IsRenderStateDirty() const;
 	bool IsVisible() const;
 	bool ShouldCreateRenderState() const;
 	void MarkRenderStateDirty();
+	void MarkPrecachePSOsRequired();
 	void PrecachePSOs();
 	bool IsPSOPrecaching() const;
 	void SetCollisionEnabled(bool bInCollisionEnabled);
@@ -193,7 +195,7 @@ protected:
 	/** Payload used to release BodyInstance resources in asynchronous mode (see OnAsyncDestroyPhysicsState). */
 	TOptional<FBodyInstance::FAsyncTermBodyPayload> AsyncTermBodyPayload;
 
-	enum class  ProxyCreationState : uint8
+	enum class EProxyCreationState : uint8
 	{
 		None, // Constructed/Initialized
 		Pending, // AddedToWorld & proxy creation is pending
@@ -202,7 +204,7 @@ protected:
 		Delayed // Proxy creation delayed (used when PSO precaching is not ready when creating proxy)
 	};
 
-	ProxyCreationState ProxyState = ProxyCreationState::None;
+	EProxyCreationState ProxyState = EProxyCreationState::None;
 	bool bRenderStateDirty = false;
 
 private:
@@ -224,9 +226,10 @@ protected:
 	std::atomic<int> LatestPSOPrecacheJobSetCompleted = 0;
 	int32 LatestPSOPrecacheJobSet = 0;
 	// Helper flag to check if PSOs have been precached already
-	uint8 bPSOPrecacheCalled : 1 = false;
+	std::atomic<bool> bPSOPrecacheCalled = false;
+	std::atomic<bool> bPSOPrecacheRequired = false;
 	// PSOs requested priority
-	EPSOPrecachePriority PSOPrecacheRequestPriority : 2 = EPSOPrecachePriority::Medium;
+	std::atomic<EPSOPrecachePriority> PSOPrecacheRequestPriority = EPSOPrecachePriority::Medium;
 	static_assert((int)EPSOPrecachePriority::Highest < 1 << 2);
 #endif
 
diff --git a/Engine/Plugins/Experimental/FastGeoStreaming/Source/FastGeoStreaming/Private/FastGeoStaticMeshComponent.cpp b/Engine/Plugins/Experimental/FastGeoStreaming/Source/FastGeoStreaming/Private/FastGeoStaticMeshComponent.cpp
--- a/Engine/Plugins/Experimental/FastGeoStreaming/Source/FastGeoStreaming/Private/FastGeoStaticMeshComponent.cpp
+++ b/Engine/Plugins/Experimental/FastGeoStreaming/Source/FastGeoStreaming/Private/FastGeoStaticMeshComponent.cpp
@@ -116,8 +116,6 @@ bool FFastGeoStaticMeshComponentBase::UseNaniteOverrideMaterials(bool bDoingMate
 
 UMaterialInterface* FFastGeoStaticMeshComponentBase::GetMaterial(int32 MaterialIndex, bool bDoingNaniteMaterialAudit) const
 {
-	TRACE_CPUPROFILER_EVENT_SCOPE(FFastGeoStaticMeshComponentBase::GetMaterial)
-
 	return FStaticMeshComponentHelper::GetMaterial(*this, MaterialIndex, bDoingNaniteMaterialAudit);
 }
 
diff --git a/Engine/Plugins/Experimental/FastGeoStreaming/Source/FastGeoStreaming/Private/FastGeoWorldSubsystem.cpp b/Engine/Plugins/Experimental/FastGeoStreaming/Source/FastGeoStreaming/Private/FastGeoWorldSubsystem.cpp
--- a/Engine/Plugins/Experimental/FastGeoStreaming/Source/FastGeoStreaming/Private/FastGeoWorldSubsystem.cpp
+++ b/Engine/Plugins/Experimental/FastGeoStreaming/Source/FastGeoStreaming/Private/FastGeoWorldSubsystem.cpp
@@ -8,16 +8,20 @@
 #include "Streaming/LevelStreamingDelegates.h"
 #include "Engine/World.h"
 #include "Engine/Level.h"
+#include "Engine/LevelStreaming.h"
 #include "WorldPartition/HLOD/HLODRuntimeSubsystem.h"
 #include "WorldPartition/WorldPartitionRuntimeCell.h"
 #include "GameFramework/ActorPrimitiveColorHandler.h"
 #include "Engine/GameViewportClient.h"
+#include "ProfilingDebugging/CsvProfiler.h"
 
 #define LOCTEXT_NAMESPACE "FastGeoWorldSubsystem"
 
 static FName NAME_FastGeoColorHandler(TEXT("FastGeo"));
 bool UFastGeoWorldSubsystem::bEnableDebugView = false;
 
+CSV_DEFINE_CATEGORY(FastGeo, true);
+
 bool UFastGeoWorldSubsystem::IsEnableDebugView()
 {
 	return bEnableDebugView;
@@ -114,6 +118,7 @@ void UFastGeoWorldSubsystem::PostInitialize()
 		FWorldDelegates::LevelAddedToWorld.AddUObject(this, &UFastGeoWorldSubsystem::OnLevelAddedToWorld);
 		FWorldDelegates::LevelRemovedFromWorld.AddUObject(this, &UFastGeoWorldSubsystem::OnLevelRemovedFromWorld);
 #endif
+		Handle_OnLevelStreamingStateChanged = FLevelStreamingDelegates::OnLevelStreamingStateChanged.AddUObject(this, &UFastGeoWorldSubsystem::OnLevelStreamingStateChanged);
 		Handle_OnLevelBeginAddToWorld = FLevelStreamingDelegates::OnLevelBeginMakingVisible.AddUObject(this, &UFastGeoWorldSubsystem::OnLevelStartedAddToWorld);
 		Handle_OnLevelBeginRemoveFromWorld = FLevelStreamingDelegates::OnLevelBeginMakingInvisible.AddUObject(this, &UFastGeoWorldSubsystem::OnLevelStartedRemoveFromWorld);
 
@@ -133,8 +138,10 @@ void UFastGeoWorldSubsystem::Deinitialize()
 	FWorldDelegates::LevelAddedToWorld.RemoveAll(this);
 	FWorldDelegates::LevelRemovedFromWorld.RemoveAll(this);
 #endif
+	FLevelStreamingDelegates::OnLevelStreamingStateChanged.Remove(Handle_OnLevelStreamingStateChanged);
 	FLevelStreamingDelegates::OnLevelBeginMakingVisible.Remove(Handle_OnLevelBeginAddToWorld);
 	FLevelStreamingDelegates::OnLevelBeginMakingInvisible.Remove(Handle_OnLevelBeginRemoveFromWorld);
+	Handle_OnLevelStreamingStateChanged.Reset();
 	Handle_OnLevelBeginAddToWorld.Reset();
 	Handle_OnLevelBeginRemoveFromWorld.Reset();
 
@@ -144,6 +151,22 @@ void UFastGeoWorldSubsystem::Deinitialize()
 	Super::Deinitialize();
 }
 
+void UFastGeoWorldSubsystem::OnLevelStreamingStateChanged(UWorld* World, const ULevelStreaming* StreamingLevel, ULevel* LevelIfLoaded, ELevelStreamingState PrevState, ELevelStreamingState NewState)
+{
+	if (World != GetWorld())
+	{
+		return;
+	}
+
+	if (LevelIfLoaded && ((NewState == ELevelStreamingState::LoadedNotVisible) || (NewState == ELevelStreamingState::LoadedVisible)))
+	{
+		if (UFastGeoContainer* FastGeo = LevelIfLoaded->GetAssetUserData<UFastGeoContainer>())
+		{
+			FastGeo->PrecachePSOs();
+		}
+	}
+}
+
 void UFastGeoWorldSubsystem::OnLevelStartedAddToWorld(UWorld* World, const ULevelStreaming* StreamingLevel, ULevel* Level)
 {
 	if (World != GetWorld())
@@ -282,10 +305,21 @@ void UFastGeoWorldSubsystem::RemoveFromComponentsPendingRecreate(FFastGeoPrimiti
 
 void UFastGeoWorldSubsystem::ProcessPendingRecreate()
 {
+	TRACE_CPUPROFILER_EVENT_SCOPE(UFastGeoWorldSubsystem::ProcessPendingRecreate);
+
+	CSV_CUSTOM_STAT(FastGeo, PendingRecreate, ComponentsPendingRecreate.Num(), ECsvCustomStatOp::Set);
+
+	int32 NumPrimitiveDelayed = 0;
 	for (FWeakFastGeoComponent WeakComponentPendingRecreate : ComponentsPendingRecreate)
 	{
 		if (FFastGeoPrimitiveComponent* ComponentPendingRecreate = WeakComponentPendingRecreate.Get<FFastGeoPrimitiveComponent>())
 		{
+#if CSV_PROFILER && !CSV_PROFILER_MINIMAL
+			if (ComponentPendingRecreate->IsRenderStateDelayed())
+			{
+				++NumPrimitiveDelayed;
+			}
+#endif
 			// Skip component if its container was unregistered before we process it
 			if (ComponentPendingRecreate->GetOwnerContainer()->IsRegistered())
 			{
@@ -295,6 +329,8 @@ void UFastGeoWorldSubsystem::ProcessPendingRecreate()
 		}
 	}
 
+	CSV_CUSTOM_STAT(FastGeo, PendingRecreateDelayed, NumPrimitiveDelayed, ECsvCustomStatOp::Set);
+
 	ComponentsPendingRecreate.Reset();
 }
 
diff --git a/Engine/Plugins/Experimental/FastGeoStreaming/Source/FastGeoStreaming/Private/FastGeoWorldSubsystem.h b/Engine/Plugins/Experimental/FastGeoStreaming/Source/FastGeoStreaming/Private/FastGeoWorldSubsystem.h
--- a/Engine/Plugins/Experimental/FastGeoStreaming/Source/FastGeoStreaming/Private/FastGeoWorldSubsystem.h
+++ b/Engine/Plugins/Experimental/FastGeoStreaming/Source/FastGeoStreaming/Private/FastGeoWorldSubsystem.h
@@ -10,6 +10,7 @@ class FFastGeoPrimitiveComponent;
 class ULevelStreaming;
 class ULevel;
 class IWorldPartitionHLODObject;
+enum class ELevelStreamingState : uint8;
 
 UCLASS()
 class FASTGEOSTREAMING_API UFastGeoWorldSubsystem : public UWorldSubsystem
@@ -37,6 +38,7 @@ public:
 
 private:
 	void OnUpdateLevelStreaming();
+	void OnLevelStreamingStateChanged(UWorld* World, const ULevelStreaming* StreamingLevel, ULevel* LevelIfLoaded, ELevelStreamingState PrevState, ELevelStreamingState NewState);
 	void OnLevelStartedAddToWorld(UWorld* World, const ULevelStreaming* StreamingLevel, ULevel* Level);
 	void OnLevelStartedRemoveFromWorld(UWorld* World, const ULevelStreaming* StreamingLevel, ULevel* Level);
 	void OnLevelComponentsUpdated(UWorld* World, ULevel* Level);
@@ -51,6 +53,7 @@ private:
 
 	void ForEachHLODObjectInCell(const UWorldPartitionRuntimeCell*, TFunction<void(IWorldPartitionHLODObject*)>);
 
+	FDelegateHandle Handle_OnLevelStreamingStateChanged;
 	FDelegateHandle Handle_OnLevelBeginAddToWorld;
 	FDelegateHandle Handle_OnLevelBeginRemoveFromWorld;
 	FDelegateHandle Handle_OnForEachHLODObjectInCell;
