FastGeo: Get rid of FTaskTagScope ETaskTag::EParallelGameThread in async thread creating the render proxies.

diff --git a/Engine/Plugins/Experimental/FastGeoStreaming/Source/FastGeoStreaming/Private/FastGeoContainer.cpp b/Engine/Plugins/Experimental/FastGeoStreaming/Source/FastGeoStreaming/Private/FastGeoContainer.cpp
--- a/Engine/Plugins/Experimental/FastGeoStreaming/Source/FastGeoStreaming/Private/FastGeoContainer.cpp
+++ b/Engine/Plugins/Experimental/FastGeoStreaming/Source/FastGeoStreaming/Private/FastGeoContainer.cpp
@@ -424,7 +424,6 @@ void UFastGeoContainer::OnCreateRenderState_Concurrent()
 	ParallelFor(NumThreads, [&](int32 Index)
 	{
 		TRACE_CPUPROFILER_EVENT_SCOPE(UFastGeoContainer::OnCreateRenderState_Concurrent);
-		FTaskTagScope Scope(ETaskTag::EParallelGameThread);
 
 		while(true)
 		{
@@ -441,7 +440,6 @@ void UFastGeoContainer::OnCreateRenderState_Concurrent()
 				return; 
 			}
 
-			// TODO_FASTGEO : Can we call GetUsedMaterials from here? FInstancedStaticMeshSceneProxy -> FStaticMeshSceneProxy -> FStaticMeshSceneProxy::FLODInfo -> FMaterial::MaterialModifiesMeshPosition_RenderThread()
 			FFastGeoPrimitiveComponent* ComponentToProcess = PendingCreate.RenderState.ComponentsToProcess[PendingCreate.RenderState.TotalNumProcessed + LocalIdx];
 			ComponentToProcess->CreateRenderState(/*Context = */ nullptr);
 
diff --git a/Engine/Source/Runtime/Engine/Private/InstancedStaticMesh.cpp b/Engine/Source/Runtime/Engine/Private/InstancedStaticMesh.cpp
--- a/Engine/Source/Runtime/Engine/Private/InstancedStaticMesh.cpp
+++ b/Engine/Source/Runtime/Engine/Private/InstancedStaticMesh.cpp
@@ -1432,7 +1432,7 @@ void FInstancedStaticMeshSceneProxy::SetupProxy(const FInstancedStaticMeshSceneP
 	}
 #endif
 
-	bDoesMeshBatchesUseSceneInstanceCount = CVarISMFetchInstanceCountFromScene.GetValueOnGameThread() != 0;
+	bDoesMeshBatchesUseSceneInstanceCount = CVarISMFetchInstanceCountFromScene.GetValueOnAnyThread() != 0;
 
 	SetupInstanceSceneDataBuffers(InstanceDataSceneProxy->GeInstanceSceneDataBuffers());
 
diff --git a/Engine/Source/Runtime/Engine/Private/Materials/MaterialShared.cpp b/Engine/Source/Runtime/Engine/Private/Materials/MaterialShared.cpp
--- a/Engine/Source/Runtime/Engine/Private/Materials/MaterialShared.cpp
+++ b/Engine/Source/Runtime/Engine/Private/Materials/MaterialShared.cpp
@@ -1063,7 +1063,6 @@ bool FMaterial::MaterialModifiesMeshPosition_RenderThread() const
 
 bool FMaterial::MaterialModifiesMeshPosition_GameThread() const
 {
-	check(IsInParallelGameThread() || IsInGameThread());
 	FMaterialShaderMap* ShaderMap = GameThreadShaderMap.GetReference();
 	return ShaderMap ? ShaderMap->ModifiesMeshPosition() : false;
 }
diff --git a/Engine/Source/Runtime/Engine/Private/StaticMeshSceneProxy.cpp b/Engine/Source/Runtime/Engine/Private/StaticMeshSceneProxy.cpp
--- a/Engine/Source/Runtime/Engine/Private/StaticMeshSceneProxy.cpp
+++ b/Engine/Source/Runtime/Engine/Private/StaticMeshSceneProxy.cpp
@@ -2504,16 +2504,16 @@ FStaticMeshSceneProxy::FLODInfo::FLODInfo(const FStaticMeshSceneProxyDesc& InPro
 		FMaterialResource const* MaterialResource = const_cast<UMaterialInterface const*>(SectionInfo.Material)->GetMaterial_Concurrent()->GetMaterialResource(FeatureLevel);
 		if(MaterialResource)
 		{
-			if (IsInParallelGameThread() || IsInGameThread())
+			if (IsInAnyRenderingThread())
 			{
-				if (MaterialResource->MaterialModifiesMeshPosition_GameThread())
+				if (MaterialResource->MaterialModifiesMeshPosition_RenderThread())
 				{
 					bUsesMeshModifyingMaterials = true;
 				}
 			}
 			else
 			{
-				if (MaterialResource->MaterialModifiesMeshPosition_RenderThread())
+				if (MaterialResource->MaterialModifiesMeshPosition_GameThread())
 				{
 					bUsesMeshModifyingMaterials = true;
 				}
