[FastGeo] Trigger GameThread completion callback even if PSO precaching has nothing to do, or else the recreate render state will never be done.
Still do the PSO precaching when using PSO creation strategy EPSOPrecacheProxyCreationStrategy::AlwaysCreate.

diff --git a/Engine/Plugins/Experimental/FastGeoStreaming/Source/FastGeoStreaming/Private/FastGeoContainer.cpp b/Engine/Plugins/Experimental/FastGeoStreaming/Source/FastGeoStreaming/Private/FastGeoContainer.cpp
--- a/Engine/Plugins/Experimental/FastGeoStreaming/Source/FastGeoStreaming/Private/FastGeoContainer.cpp
+++ b/Engine/Plugins/Experimental/FastGeoStreaming/Source/FastGeoStreaming/Private/FastGeoContainer.cpp
@@ -668,7 +668,7 @@ void UFastGeoContainer::PrecachePSOs()
 	}
 	bPrecachedPSOs = true;
 
-	if (IsComponentPSOPrecachingEnabled() && (GetPSOPrecacheProxyCreationStrategy() != EPSOPrecacheProxyCreationStrategy::AlwaysCreate))
+	if (IsComponentPSOPrecachingEnabled())
 	{
 		TArray<FFastGeoPrimitiveComponent*> Components;
 		ForEachComponentCluster([&Components](FFastGeoComponentCluster& ComponentCluster)
diff --git a/Engine/Plugins/Experimental/FastGeoStreaming/Source/FastGeoStreaming/Private/FastGeoPrimitiveComponent.cpp b/Engine/Plugins/Experimental/FastGeoStreaming/Source/FastGeoStreaming/Private/FastGeoPrimitiveComponent.cpp
--- a/Engine/Plugins/Experimental/FastGeoStreaming/Source/FastGeoStreaming/Private/FastGeoPrimitiveComponent.cpp
+++ b/Engine/Plugins/Experimental/FastGeoStreaming/Source/FastGeoStreaming/Private/FastGeoPrimitiveComponent.cpp
@@ -619,15 +619,8 @@ void FFastGeoPrimitiveComponent::RequestRecreateRenderStateWhenPSOPrecacheFinish
 	if (GetPSOPrecacheProxyCreationStrategy() != EPSOPrecacheProxyCreationStrategy::AlwaysCreate)
 	{
 		LatestPSOPrecacheJobSet++;
-		if (!PSOPrecacheCompileEvents.IsEmpty())
-		{
-			TGraphTask<FFastGeoPSOPrecacheFinishedTask>::CreateTask(&PSOPrecacheCompileEvents).ConstructAndDispatchWhenReady(this, LatestPSOPrecacheJobSet);
-		}
-		else
-		{
-			// No graph events to wait on, the job set can be considered complete.
-			LatestPSOPrecacheJobSetCompleted = LatestPSOPrecacheJobSet;
-		}
+		// Even if PSOPrecacheCompileEvents is empty, still push the completion task as it needs to run on the Game Thread and call MarkRenderStateDirty
+		TGraphTask<FFastGeoPSOPrecacheFinishedTask>::CreateTask(&PSOPrecacheCompileEvents).ConstructAndDispatchWhenReady(this, LatestPSOPrecacheJobSet);
 	}
 	bPSOPrecacheCalled = true;
 }
