[SimpleStreamableAssetManager][FastGeo] Put back FTaskTagScope ETaskTag::EParallelGameThread in UFastGeoContainer::OnCreateRenderState_Concurrent as it's currently needed for FastGeo to use SimpleStreamableAssetManager.

diff --git a/Engine/Plugins/Experimental/FastGeoStreaming/Source/FastGeoStreaming/Private/FastGeoContainer.cpp b/Engine/Plugins/Experimental/FastGeoStreaming/Source/FastGeoStreaming/Private/FastGeoContainer.cpp
--- a/Engine/Plugins/Experimental/FastGeoStreaming/Source/FastGeoStreaming/Private/FastGeoContainer.cpp
+++ b/Engine/Plugins/Experimental/FastGeoStreaming/Source/FastGeoStreaming/Private/FastGeoContainer.cpp
@@ -424,6 +424,8 @@ void UFastGeoContainer::OnCreateRenderState_Concurrent()
 	ParallelFor(NumThreads, [&](int32 Index)
 	{
 		TRACE_CPUPROFILER_EVENT_SCOPE(UFastGeoContainer::OnCreateRenderState_Concurrent);
+		// Currently necessary for FSimpleStreamableAssetManager::FRegister (to be revisited)
+		FTaskTagScope Scope(ETaskTag::EParallelGameThread);
 
 		while(true)
 		{
