[SimpleStreamableAssetManager][FastGeo] Persist new SceneProxyDesc bForceMipStreaming in FastGeo primitive components.

diff --git a/Engine/Plugins/Experimental/FastGeoStreaming/Source/FastGeoStreaming/Private/FastGeoPrimitiveComponent.cpp b/Engine/Plugins/Experimental/FastGeoStreaming/Source/FastGeoStreaming/Private/FastGeoPrimitiveComponent.cpp
--- a/Engine/Plugins/Experimental/FastGeoStreaming/Source/FastGeoStreaming/Private/FastGeoPrimitiveComponent.cpp
+++ b/Engine/Plugins/Experimental/FastGeoStreaming/Source/FastGeoStreaming/Private/FastGeoPrimitiveComponent.cpp
@@ -118,6 +118,7 @@ void FFastGeoPrimitiveComponent::Serialize(FArchive& Ar)
 	FArchive_Serialize_BitfieldBool(Ar, SceneProxyDesc.bRenderCustomDepth);
 	FArchive_Serialize_BitfieldBool(Ar, SceneProxyDesc.bVisibleInSceneCaptureOnly);
 	FArchive_Serialize_BitfieldBool(Ar, SceneProxyDesc.bHiddenInSceneCapture);
+	FArchive_Serialize_BitfieldBool(Ar, SceneProxyDesc.bForceMipStreaming);
 	FArchive_Serialize_BitfieldBool(Ar, SceneProxyDesc.bRayTracingFarField);
 	FArchive_Serialize_BitfieldBool(Ar, SceneProxyDesc.bHoldout);
 	FArchive_Serialize_BitfieldBool(Ar, SceneProxyDesc.bIsFirstPerson);
