World partition cell transformer - Added support for excluding actors that have the "CellTransformerIgnoreActor" actor tag
* Updated ISM and FastGeo transformer to perform that test

diff --git a/Engine/Plugins/Experimental/FastGeoStreaming/Source/FastGeoStreaming/Private/FastGeoWorldPartitionRuntimeCellTransformer.cpp b/Engine/Plugins/Experimental/FastGeoStreaming/Source/FastGeoStreaming/Private/FastGeoWorldPartitionRuntimeCellTransformer.cpp
--- a/Engine/Plugins/Experimental/FastGeoStreaming/Source/FastGeoStreaming/Private/FastGeoWorldPartitionRuntimeCellTransformer.cpp
+++ b/Engine/Plugins/Experimental/FastGeoStreaming/Source/FastGeoStreaming/Private/FastGeoWorldPartitionRuntimeCellTransformer.cpp
@@ -597,12 +597,6 @@ FFastGeoTransformResult UFastGeoWorldPartitionRuntimeCellTransformer::IsAllowedA
 		}
 	}
 
-	// Special case where we disallow an actor class if actor is tagged 'NoFastGeo'
-	if (InActor->Tags.Contains(FastGeo::NAME_NoFastGeo))
-	{
-		return { EFastGeoTransform::Reject, [InActor] { return FString::Printf(TEXT("Actor %s is tagged 'NoFastGeo'"), *InActor->GetName()); } };
-	}
-
 	for (TSubclassOf<AActor> AllowedActorClass : AllowedActorClasses)
 	{
 		if (ActorClass->IsChildOf(AllowedActorClass))
@@ -647,6 +641,17 @@ FFastGeoTransformResult UFastGeoWorldPartitionRuntimeCellTransformer::CanTransfo
 	}
 
 	FString Reason;
+	
+	if (InActor->ActorHasTag(NAME_CellTransformerIgnoreActor))
+	{
+		return { EFastGeoTransform::Reject, [InActor] { return FString::Printf(TEXT("Actor %s is tagged '%s'"), *InActor->GetName(), *NAME_CellTransformerIgnoreActor.ToString()); } };
+	}
+	
+	if (InActor->ActorHasTag(FastGeo::NAME_NoFastGeo))
+	{
+		return { EFastGeoTransform::Reject, [InActor] { return FString::Printf(TEXT("Actor %s is tagged '%s'"), *InActor->GetName(), *FastGeo::NAME_NoFastGeo.ToString()); } };
+	}
+	
 	if (!IsActorTransformable(InActor, Reason))
 	{
 		return { EFastGeoTransform::Reject, [InActor, Reason] { return FString::Printf(TEXT("Actor %s [%s]"), *InActor->GetName(), *Reason); } };
diff --git a/Engine/Source/Runtime/Engine/Private/WorldPartition/WorldPartitionRuntimeCellTransformer.cpp b/Engine/Source/Runtime/Engine/Private/WorldPartition/WorldPartitionRuntimeCellTransformer.cpp
--- a/Engine/Source/Runtime/Engine/Private/WorldPartition/WorldPartitionRuntimeCellTransformer.cpp
+++ b/Engine/Source/Runtime/Engine/Private/WorldPartition/WorldPartitionRuntimeCellTransformer.cpp
@@ -4,4 +4,6 @@
 
 UWorldPartitionRuntimeCellTransformer::UWorldPartitionRuntimeCellTransformer(const FObjectInitializer& ObjectInitializer)
 	: Super(ObjectInitializer)
-{}
\ No newline at end of file
+{}
+
+const FName UWorldPartitionRuntimeCellTransformer::NAME_CellTransformerIgnoreActor(TEXT("CellTransformer_IgnoreActor"));
diff --git a/Engine/Source/Runtime/Engine/Private/WorldPartition/WorldPartitionRuntimeCellTransformerISM.cpp b/Engine/Source/Runtime/Engine/Private/WorldPartition/WorldPartitionRuntimeCellTransformerISM.cpp
--- a/Engine/Source/Runtime/Engine/Private/WorldPartition/WorldPartitionRuntimeCellTransformerISM.cpp
+++ b/Engine/Source/Runtime/Engine/Private/WorldPartition/WorldPartitionRuntimeCellTransformerISM.cpp
@@ -137,6 +137,11 @@ void UWorldPartitionRuntimeCellTransformerISM::Transform(ULevel* InLevel)
 
 bool UWorldPartitionRuntimeCellTransformerISM::CanAutoInstanceActor(AActor* InActor) const
 {
+	if (InActor->ActorHasTag(NAME_CellTransformerIgnoreActor))
+	{
+		return false;
+	}
+	
 	if (InActor->GetIsReplicated())
 	{
 		return false;
diff --git a/Engine/Source/Runtime/Engine/Public/WorldPartition/WorldPartitionRuntimeCellTransformer.h b/Engine/Source/Runtime/Engine/Public/WorldPartition/WorldPartitionRuntimeCellTransformer.h
--- a/Engine/Source/Runtime/Engine/Public/WorldPartition/WorldPartitionRuntimeCellTransformer.h
+++ b/Engine/Source/Runtime/Engine/Public/WorldPartition/WorldPartitionRuntimeCellTransformer.h
@@ -17,6 +17,10 @@ public:
 #endif
 
 	bool IsEnabled() const { return bEnabled; }
+	
+protected:
+	// Tag used to force exclude actors from any cell transformation
+	ENGINE_API static const FName NAME_CellTransformerIgnoreActor;
 
 private:
 	UPROPERTY(EditAnywhere, Category = Transformer)
