FastGeo - ShouldCreateRenderState() should consider the full set of options that require the creation of a scene proxy

diff --git a/Engine/Plugins/Experimental/FastGeoStreaming/Source/FastGeoStreaming/Private/FastGeoComponentCluster.cpp b/Engine/Plugins/Experimental/FastGeoStreaming/Source/FastGeoStreaming/Private/FastGeoComponentCluster.cpp
--- a/Engine/Plugins/Experimental/FastGeoStreaming/Source/FastGeoStreaming/Private/FastGeoComponentCluster.cpp
+++ b/Engine/Plugins/Experimental/FastGeoStreaming/Source/FastGeoStreaming/Private/FastGeoComponentCluster.cpp
@@ -102,14 +102,14 @@ void FFastGeoComponentCluster::UpdateVisibility()
 	TArray<FPrimitiveSceneProxy*, TInlineAllocator<32>> HideProxies;
 	ForEachComponent<FFastGeoPrimitiveComponent>([&ShowProxies, &HideProxies](FFastGeoPrimitiveComponent& Component)
 	{
-		const bool bOldIsVisible = Component.IsVisible();
+		const bool bOldIsDrawnInGame = Component.IsDrawnInGame();
 		Component.UpdateVisibility();
-		const bool bIsVisible = Component.IsVisible();
-		if (bIsVisible != bOldIsVisible)
+		const bool bIsDrawnInGame = Component.IsDrawnInGame();
+		if (bIsDrawnInGame != bOldIsDrawnInGame)
 		{
 			if (FPrimitiveSceneProxy* SceneProxy = Component.GetSceneProxy())
 			{
-				if (bIsVisible)
+				if (bIsDrawnInGame)
 				{
 					ShowProxies.Add(SceneProxy);
 				}
diff --git a/Engine/Plugins/Experimental/FastGeoStreaming/Source/FastGeoStreaming/Private/FastGeoPrimitiveComponent.cpp b/Engine/Plugins/Experimental/FastGeoStreaming/Source/FastGeoStreaming/Private/FastGeoPrimitiveComponent.cpp
--- a/Engine/Plugins/Experimental/FastGeoStreaming/Source/FastGeoStreaming/Private/FastGeoPrimitiveComponent.cpp
+++ b/Engine/Plugins/Experimental/FastGeoStreaming/Source/FastGeoStreaming/Private/FastGeoPrimitiveComponent.cpp
@@ -234,9 +234,11 @@ void FFastGeoPrimitiveComponent::ResetSceneProxyDescUnsupportedProperties()
 }
 #endif
 
-bool FFastGeoPrimitiveComponent::IsVisible() const
+bool FFastGeoPrimitiveComponent::IsDrawnInGame() const
 {
-	return GetSceneProxyDesc().bIsVisible;
+	// Drawn in game must consider both the component bIsVisible flag AND the bIsHidden flag (which actually originates from the actor bHiddenInGame property)
+	// This logic mimics what is done to initialize FPrimitiveSceneProxy::DrawInGame
+	return GetSceneProxyDesc().bIsVisible && !GetSceneProxyDesc().bIsHidden;
 }
 
 EComponentMobility::Type FFastGeoPrimitiveComponent::GetMobility() const
@@ -500,7 +502,16 @@ bool FFastGeoPrimitiveComponent::IsRenderStateDirty() const
 
 bool FFastGeoPrimitiveComponent::ShouldCreateRenderState() const
 {
-	return bIsVisible;
+	// If the detail mode setting allows it, add it to the scene.
+	const bool bDetailModeAllowsRendering = DetailMode <= GetCachedScalabilityCVars().DetailMode;
+	if (!bDetailModeAllowsRendering)
+	{
+		return false;
+	}
+
+	const FPrimitiveSceneProxyDesc& SceneProxyDesc = GetSceneProxyDesc();
+	const bool bShouldCreateRenderState = (bIsVisible && !SceneProxyDesc.bIsHidden) || SceneProxyDesc.bCastHiddenShadow || SceneProxyDesc.bAffectIndirectLightingWhileHidden || SceneProxyDesc.bRayTracingFarField;
+	return bShouldCreateRenderState;
 }
 
 bool FFastGeoPrimitiveComponent::IsFirstPersonRelevant() const
diff --git a/Engine/Plugins/Experimental/FastGeoStreaming/Source/FastGeoStreaming/Private/FastGeoPrimitiveComponent.h b/Engine/Plugins/Experimental/FastGeoStreaming/Source/FastGeoStreaming/Private/FastGeoPrimitiveComponent.h
--- a/Engine/Plugins/Experimental/FastGeoStreaming/Source/FastGeoStreaming/Private/FastGeoPrimitiveComponent.h
+++ b/Engine/Plugins/Experimental/FastGeoStreaming/Source/FastGeoStreaming/Private/FastGeoPrimitiveComponent.h
@@ -111,7 +111,7 @@ public:
 	bool IsRenderStateCreated() const;
 	bool IsRenderStateDelayed() const;
 	bool IsRenderStateDirty() const;
-	bool IsVisible() const;
+	bool IsDrawnInGame() const;
 	bool ShouldCreateRenderState() const;
 	void MarkRenderStateDirty();
 	void MarkPrecachePSOsRequired();
