FastGeo - Fix missing include.

diff --git a/Engine/Plugins/Experimental/FastGeoStreaming/Source/FastGeoStreaming/Private/FastGeoPrimitiveComponent.cpp b/Engine/Plugins/Experimental/FastGeoStreaming/Source/FastGeoStreaming/Private/FastGeoPrimitiveComponent.cpp
--- a/Engine/Plugins/Experimental/FastGeoStreaming/Source/FastGeoStreaming/Private/FastGeoPrimitiveComponent.cpp
+++ b/Engine/Plugins/Experimental/FastGeoStreaming/Source/FastGeoStreaming/Private/FastGeoPrimitiveComponent.cpp
@@ -21,6 +21,7 @@
 #include "PSOPrecache.h"
 #include "PSOPrecacheMaterial.h"
 #include "PrimitiveComponentHelper.h"
+#include "UnrealEngine.h"
 
 #if !UE_BUILD_SHIPPING
 namespace FastGeo
