FastGeo - Use a unified frame budget for the async render state tasks

diff --git a/Engine/Plugins/Experimental/FastGeoStreaming/Source/FastGeoStreaming/Private/FastGeoAsyncRenderStateJobQueue.cpp b/Engine/Plugins/Experimental/FastGeoStreaming/Source/FastGeoStreaming/Private/FastGeoAsyncRenderStateJobQueue.cpp
--- a/Engine/Plugins/Experimental/FastGeoStreaming/Source/FastGeoStreaming/Private/FastGeoAsyncRenderStateJobQueue.cpp
+++ b/Engine/Plugins/Experimental/FastGeoStreaming/Source/FastGeoStreaming/Private/FastGeoAsyncRenderStateJobQueue.cpp
@@ -2,6 +2,7 @@
 
 #include "FastGeoAsyncRenderStateJobQueue.h"
 #include "FastGeoContainer.h"
+#include "FastGeoWorldSubsystem.h"
 #include "Engine/World.h"
 #include "Engine/Level.h"
 #include "SceneInterface.h"
@@ -122,10 +123,7 @@ void FFastGeoAsyncRenderStateJobQueue::Tick(bool bForceWaitCompletion)
 	Launch();
 
 	// Wait for tasks if the world is inside a blocking load
-	if (bForceWaitCompletion ||
-		Scene->GetWorld()->GetIsInBlockTillLevelStreamingCompleted() ||
-		Scene->GetWorld()->GetShouldForceUnloadStreamingLevels() ||
-		Scene->GetWorld()->IsBeingCleanedUp())
+	if (Scene->GetWorld()->GetSubsystem<UFastGeoWorldSubsystem>()->IsWaitingForCompletion())
 	{
 		WaitForAsyncTasksExecution();
 	}
diff --git a/Engine/Plugins/Experimental/FastGeoStreaming/Source/FastGeoStreaming/Private/FastGeoContainer.cpp b/Engine/Plugins/Experimental/FastGeoStreaming/Source/FastGeoStreaming/Private/FastGeoContainer.cpp
--- a/Engine/Plugins/Experimental/FastGeoStreaming/Source/FastGeoStreaming/Private/FastGeoContainer.cpp
+++ b/Engine/Plugins/Experimental/FastGeoStreaming/Source/FastGeoStreaming/Private/FastGeoContainer.cpp
@@ -5,6 +5,7 @@
 #include "FastGeoHLOD.h"
 #include "FastGeoComponent.h"
 #include "FastGeoStreamingManager.h"
+#include "FastGeoWorldSubsystem.h"
 #include "FastGeoLog.h"
 #include "AI/Navigation/NavigationElement.h"
 #include "Async/ParallelFor.h"
@@ -34,20 +35,6 @@ namespace FastGeo
 		TEXT("Only taken into account if value is greater than 1."),
 		ECVF_Default);
 
-	static float GAsyncRenderStateTaskTimeBudgetMS = 0.0f;
-	static FAutoConsoleVariableRef CVarAsyncRenderStateTaskTimeBudgetMS(
-		TEXT("FastGeo.AsyncRenderStateTask.TimeBudgetMS"),
-		GAsyncRenderStateTaskTimeBudgetMS,
-		TEXT("Maximum time budget in milliseconds for the async render state tasks (0 = no time limit)"),
-		ECVF_Default);
-
-	static int32 GAsyncRenderStateTaskMaxNumComponentsToProcess = 0;
-	static FAutoConsoleVariableRef CVarAsyncRenderStateTaskMaxNumComponentsToProcess(
-		TEXT("FastGeo.AsyncRenderStateTask.MaxNumComponentsToProcess"),
-		GAsyncRenderStateTaskMaxNumComponentsToProcess,
-		TEXT("Maximum number of components to process (0 = no component limit)"),
-		ECVF_Default);
-
 	class FAssetRemapArchive : public FArchiveProxy
 	{
 	public:
@@ -391,11 +378,9 @@ void UFastGeoContainer::OnCreateRenderStateBegin_GameThread()
 	check(PendingCreate.RenderState.NumToProcess == 0);
 	check(PendingCreate.RenderState.TotalNumProcessed != PendingCreate.RenderState.ComponentsToProcess.Num());
 
+	PendingCreate.bIsInBlockingWait = GetWorld()->GetSubsystem<UFastGeoWorldSubsystem>()->IsWaitingForCompletion();
+
 	PendingCreate.RenderState.NumToProcess = PendingCreate.RenderState.ComponentsToProcess.Num() - PendingCreate.RenderState.TotalNumProcessed;
-	if (FastGeo::GAsyncRenderStateTaskMaxNumComponentsToProcess > 0)
-	{
-		PendingCreate.RenderState.NumToProcess = FMath::Min(FastGeo::GAsyncRenderStateTaskMaxNumComponentsToProcess, PendingCreate.RenderState.NumToProcess);
-	}
 }
 
 void UFastGeoContainer::OnDestroyRenderStateBegin_GameThread()
@@ -403,6 +388,8 @@ void UFastGeoContainer::OnDestroyRenderStateBegin_GameThread()
 	check(PendingDestroy.RenderState.NumToProcess == 0);
 	check(PendingDestroy.RenderState.TotalNumProcessed != PendingDestroy.RenderState.ComponentsToProcess.Num());
 
+	PendingDestroy.bIsInBlockingWait = GetWorld()->GetSubsystem<UFastGeoWorldSubsystem>()->IsWaitingForCompletion();
+
 	// TODO_FASTGEO: Evaluate if we need throttling of the async destruction task
 	PendingDestroy.RenderState.NumToProcess = PendingDestroy.RenderState.ComponentsToProcess.Num();
 }
@@ -411,46 +398,59 @@ void UFastGeoContainer::OnCreateRenderState_Concurrent()
 {
 	static const int32 MinNumElementsToProcessPerThread = 8;
 	const int32 NumComponentsToProcess = PendingCreate.RenderState.NumToProcess;
-	const int32 NumThreads = FMath::Clamp(NumComponentsToProcess / MinNumElementsToProcessPerThread, 1, FastGeo::GAsyncRenderStateTaskParallelWorkerCount);
+	const int32 MaxNumThreads = PendingCreate.bIsInBlockingWait ? INT32_MAX : FastGeo::GAsyncRenderStateTaskParallelWorkerCount;
+	const int32 NumThreads = FMath::Clamp(NumComponentsToProcess / MinNumElementsToProcessPerThread, 1, MaxNumThreads);
 	const bool bIsParallelForAllowed = NumThreads > 1 && FApp::ShouldUseThreadingForPerformance();
 
-	const bool bTimeLimit = FastGeo::GAsyncRenderStateTaskTimeBudgetMS != 0;
-	const double TimeBudgetSeconds = FastGeo::GAsyncRenderStateTaskTimeBudgetMS / 1000.0;
-	UE::FTimeout Timeout = bTimeLimit ? UE::FTimeout(TimeBudgetSeconds) : UE::FTimeout::Never();
+	float AvailableTimeBudgetMS;
+	int32 AvailableComponentsBudget;
+	int32 TimeEpoch;
+
+	UFastGeoWorldSubsystem* WorldSubsystem = GetWorld()->GetSubsystem<UFastGeoWorldSubsystem>();
+	WorldSubsystem->RequestAsyncRenderStateTasksBudget_Concurrent(AvailableTimeBudgetMS, AvailableComponentsBudget, TimeEpoch);
+	
+	const int32 ComponentsBudget = FMath::Min(NumComponentsToProcess, AvailableComponentsBudget);
+	const double TimeBudgetSeconds = AvailableTimeBudgetMS / 1000.0;
+	UE::FTimeout Timeout = UE::FTimeout(TimeBudgetSeconds);
 
 	TAtomic<int32> NextIndex = 0;
 	TAtomic<int32> NumProcessed = 0;
 
-	ParallelFor(NumThreads, [&](int32 Index)
+	if (ComponentsBudget > 0 && !Timeout.IsExpired())
 	{
-		TRACE_CPUPROFILER_EVENT_SCOPE(UFastGeoContainer::OnCreateRenderState_Concurrent);
-		// Currently necessary for FSimpleStreamableAssetManager::FRegister (to be revisited)
-		FTaskTagScope Scope(ETaskTag::EParallelGameThread);
-
-		while(true)
+		ParallelFor(NumThreads, [&](int32 Index)
 		{
-			// Time budget exceeded ?
-			if (Timeout.IsExpired())
-			{
-				return; 
-			}
+			TRACE_CPUPROFILER_EVENT_SCOPE(UFastGeoContainer::OnCreateRenderState_Concurrent);
+			// Currently necessary for FSimpleStreamableAssetManager::FRegister (to be revisited)
+			FTaskTagScope Scope(ETaskTag::EParallelGameThread);
 
-			// All work completed ?
-			int32 LocalIdx = NextIndex++;
-			if (LocalIdx >= NumComponentsToProcess)
+			while(true)
 			{
-				return; 
-			}
+				// Time budget exceeded ?
+				if (Timeout.IsExpired())
+				{
+					return; 
+				}
 
-			FFastGeoPrimitiveComponent* ComponentToProcess = PendingCreate.RenderState.ComponentsToProcess[PendingCreate.RenderState.TotalNumProcessed + LocalIdx];
-			ComponentToProcess->CreateRenderState(/*Context = */ nullptr);
+				// All work completed ?
+				int32 LocalIdx = NextIndex++;
+				if (LocalIdx >= ComponentsBudget)
+				{
+					return; 
+				}
 
-			NumProcessed++;
-		}
-	}, bIsParallelForAllowed ? EParallelForFlags::None : EParallelForFlags::ForceSingleThread);
+				FFastGeoPrimitiveComponent* ComponentToProcess = PendingCreate.RenderState.ComponentsToProcess[PendingCreate.RenderState.TotalNumProcessed + LocalIdx];
+				ComponentToProcess->CreateRenderState(/*Context = */ nullptr);
+
+				NumProcessed++;
+			}
+		}, bIsParallelForAllowed ? EParallelForFlags::None : EParallelForFlags::ForceSingleThread);
+	}
 
 	// Adjust NumRenderStatesToProcess to the actual number of states we processed. 
 	PendingCreate.RenderState.NumProcessed = NumProcessed.Load(EMemoryOrder::Relaxed);
+
+	WorldSubsystem->CommitAsyncRenderStateTasksBudget_Concurrent(Timeout.GetElapsedSeconds() * 1000, PendingCreate.RenderState.NumProcessed, TimeEpoch);
 }
 
 void UFastGeoContainer::OnDestroyRenderState_Concurrent()
@@ -479,6 +479,7 @@ void UFastGeoContainer::OnCreateRenderStateEnd_GameThread()
 	PendingCreate.RenderState.TotalNumProcessed += PendingCreate.RenderState.NumProcessed;
 	PendingCreate.RenderState.NumToProcess = 0;
 	PendingCreate.RenderState.NumProcessed = 0;
+	PendingCreate.bIsInBlockingWait = false;
 
 	if (PendingCreate.RenderState.TotalNumProcessed < PendingCreate.RenderState.ComponentsToProcess.Num())
 	{
@@ -494,6 +495,7 @@ void UFastGeoContainer::OnDestroyRenderStateEnd_GameThread()
 	PendingDestroy.RenderState.TotalNumProcessed += PendingDestroy.RenderState.NumProcessed;
 	PendingDestroy.RenderState.NumToProcess = 0;
 	PendingDestroy.RenderState.NumProcessed = 0;	
+	PendingDestroy.bIsInBlockingWait = false;
 
 	if (PendingDestroy.RenderState.TotalNumProcessed < PendingDestroy.RenderState.ComponentsToProcess.Num())
 	{
diff --git a/Engine/Plugins/Experimental/FastGeoStreaming/Source/FastGeoStreaming/Private/FastGeoContainer.h b/Engine/Plugins/Experimental/FastGeoStreaming/Source/FastGeoStreaming/Private/FastGeoContainer.h
--- a/Engine/Plugins/Experimental/FastGeoStreaming/Source/FastGeoStreaming/Private/FastGeoContainer.h
+++ b/Engine/Plugins/Experimental/FastGeoStreaming/Source/FastGeoStreaming/Private/FastGeoContainer.h
@@ -37,6 +37,8 @@ struct FFastGeoComponentClusterPendingState
 	} RenderState;
 
 	bool bPhysicsState;
+
+	bool bIsInBlockingWait;
 	
 	FFastGeoComponentClusterPendingState()
 	{
@@ -49,7 +51,7 @@ struct FFastGeoComponentClusterPendingState
 		RenderState.NumToProcess = 0;
 		RenderState.NumProcessed = 0;
 		RenderState.TotalNumProcessed = 0;
-
+		bIsInBlockingWait = false;
 		bPhysicsState = false;
 	}
 
diff --git a/Engine/Plugins/Experimental/FastGeoStreaming/Source/FastGeoStreaming/Private/FastGeoWorldSubsystem.cpp b/Engine/Plugins/Experimental/FastGeoStreaming/Source/FastGeoStreaming/Private/FastGeoWorldSubsystem.cpp
--- a/Engine/Plugins/Experimental/FastGeoStreaming/Source/FastGeoStreaming/Private/FastGeoWorldSubsystem.cpp
+++ b/Engine/Plugins/Experimental/FastGeoStreaming/Source/FastGeoStreaming/Private/FastGeoWorldSubsystem.cpp
@@ -20,6 +20,23 @@
 static FName NAME_FastGeoColorHandler(TEXT("FastGeo"));
 bool UFastGeoWorldSubsystem::bEnableDebugView = false;
 
+namespace FastGeo
+{
+	static float GAsyncRenderStateTaskTimeBudgetMS = 0.0f;
+	static FAutoConsoleVariableRef CVarAsyncRenderStateTaskTimeBudgetMS(
+		TEXT("FastGeo.AsyncRenderStateTask.TimeBudgetMS"),
+		GAsyncRenderStateTaskTimeBudgetMS,
+		TEXT("Maximum time budget in milliseconds for the async render state tasks (0 = no time limit)"),
+		ECVF_Default);
+
+	static int32 GAsyncRenderStateTaskMaxNumComponentsToProcess = 0;
+	static FAutoConsoleVariableRef CVarAsyncRenderStateTaskMaxNumComponentsToProcess(
+		TEXT("FastGeo.AsyncRenderStateTask.MaxNumComponentsToProcess"),
+		GAsyncRenderStateTaskMaxNumComponentsToProcess,
+		TEXT("Maximum number of components to process (0 = no component limit)"),
+		ECVF_Default);
+}
+
 CSV_DEFINE_CATEGORY(FastGeo, true);
 
 bool UFastGeoWorldSubsystem::IsEnableDebugView()
@@ -241,15 +258,17 @@ void UFastGeoWorldSubsystem::OnLevelComponentsCleared(UWorld* World, ULevel* Lev
 	}
 }
 
-void UFastGeoWorldSubsystem::OnAddLevelToWorldExtension(ULevel* Level, const bool bWaitForCompletion, bool& bOutHasCompleted)
+void UFastGeoWorldSubsystem::OnAddLevelToWorldExtension(ULevel* InLevel, const bool bInWaitForCompletion, bool& bOutHasCompleted)
 {
-	if (UFastGeoContainer* FastGeo = Level->GetAssetUserData<UFastGeoContainer>())
+	TGuardValue<bool> GuardValue(bWaitingForCompletion, bInWaitForCompletion);
+
+	if (UFastGeoContainer* FastGeo = InLevel->GetAssetUserData<UFastGeoContainer>())
 	{
 		TRACE_CPUPROFILER_EVENT_SCOPE(UFastGeoWorldSubsystem::OnAddLevelToWorldExtension);
 
 		if (FastGeo->HasAnyPendingCreateTasks())
 		{
-			FastGeo->Tick(bWaitForCompletion);
+			FastGeo->Tick(bInWaitForCompletion);
 		}
 
 		if (FastGeo->HasAnyPendingCreateTasks())
@@ -259,15 +278,17 @@ void UFastGeoWorldSubsystem::OnAddLevelToWorldExtension(ULevel* Level, const boo
 	}
 }
 
-void UFastGeoWorldSubsystem::OnRemoveLevelFromWorldExtension(ULevel* Level, const bool bWaitForCompletion, bool& bOutHasCompleted)
+void UFastGeoWorldSubsystem::OnRemoveLevelFromWorldExtension(ULevel* InLevel, const bool bInWaitForCompletion, bool& bOutHasCompleted)
 {
-	if (UFastGeoContainer* FastGeo = Level->GetAssetUserData<UFastGeoContainer>())
+	TGuardValue<bool> GuardValue(bWaitingForCompletion, bInWaitForCompletion);
+
+	if (UFastGeoContainer* FastGeo = InLevel->GetAssetUserData<UFastGeoContainer>())
 	{
 		TRACE_CPUPROFILER_EVENT_SCOPE(UFastGeoWorldSubsystem::OnRemoveLevelFromWorldExtension);
 
 		if (FastGeo->HasAnyPendingDestroyTasks())
 		{
-			FastGeo->Tick(bWaitForCompletion);
+			FastGeo->Tick(bInWaitForCompletion);
 		}
 
 		if (FastGeo->HasAnyPendingDestroyTasks())
@@ -387,4 +408,71 @@ void UFastGeoWorldSubsystem::OnLevelRemovedFromWorld(ULevel* Level, UWorld* Worl
 }
 #endif
 
+bool UFastGeoWorldSubsystem::IsWaitingForCompletion() const
+{
+	return bWaitingForCompletion ||
+			GetWorld()->GetIsInBlockTillLevelStreamingCompleted() ||
+			GetWorld()->GetShouldForceUnloadStreamingLevels() ||
+			GetWorld()->IsBeingCleanedUp();
+}
+
+void UFastGeoWorldSubsystem::Tick(float DeltaTime)
+{
+	FWriteScopeLock ScopeLock(Lock);
+
+	++TimeEpoch;
+	UsedAsyncRenderStateTasksTimeBudgetMS = 0;
+	UsedNumComponentsToProcessBudget = 0;
+}
+
+TStatId UFastGeoWorldSubsystem::GetStatId() const
+{
+	RETURN_QUICK_DECLARE_CYCLE_STAT(UFastGeoWorldSubsystem, STATGROUP_Tickables);
+}
+
+void UFastGeoWorldSubsystem::RequestAsyncRenderStateTasksBudget_Concurrent(float& OutAvailableTimeBudgetMS, int32& OutAvailableComponentsBudget, int32& OutTimeEpoch)
+{
+	FWriteScopeLock ScopeLock(Lock);
+
+	bool bUnlimitedBudget = IsInGameThread() ? IsWaitingForCompletion() : false;
+
+	if (bUnlimitedBudget || FastGeo::GAsyncRenderStateTaskTimeBudgetMS == 0)
+	{
+		OutAvailableTimeBudgetMS = FLT_MAX;
+	}
+	else
+	{
+		OutAvailableTimeBudgetMS = FMath::Max(FastGeo::GAsyncRenderStateTaskTimeBudgetMS - UsedAsyncRenderStateTasksTimeBudgetMS, 0);
+	}
+
+	if (bUnlimitedBudget || FastGeo::GAsyncRenderStateTaskMaxNumComponentsToProcess == 0)
+	{
+		OutAvailableComponentsBudget = INT32_MAX;
+	}
+	else
+	{
+		OutAvailableComponentsBudget = FMath::Max(FastGeo::GAsyncRenderStateTaskMaxNumComponentsToProcess - UsedNumComponentsToProcessBudget, 0);
+	}
+
+	OutTimeEpoch = TimeEpoch;
+}
+
+void UFastGeoWorldSubsystem::CommitAsyncRenderStateTasksBudget_Concurrent(float InUsedTimeBudgetMS, int32& InUsedComponentsBudget, int32 InTimeEpoch)
+{
+	FWriteScopeLock ScopeLock(Lock);
+
+	if (InTimeEpoch == TimeEpoch)
+	{
+		if (FastGeo::GAsyncRenderStateTaskTimeBudgetMS != 0)
+		{
+			UsedAsyncRenderStateTasksTimeBudgetMS += InUsedTimeBudgetMS;
+		}
+
+		if (FastGeo::GAsyncRenderStateTaskMaxNumComponentsToProcess != 0)
+		{
+			UsedNumComponentsToProcessBudget += InUsedComponentsBudget;
+		}
+	}
+}
+
 #undef LOCTEXT_NAMESPACE
\ No newline at end of file
diff --git a/Engine/Plugins/Experimental/FastGeoStreaming/Source/FastGeoStreaming/Private/FastGeoWorldSubsystem.h b/Engine/Plugins/Experimental/FastGeoStreaming/Source/FastGeoStreaming/Private/FastGeoWorldSubsystem.h
--- a/Engine/Plugins/Experimental/FastGeoStreaming/Source/FastGeoStreaming/Private/FastGeoWorldSubsystem.h
+++ b/Engine/Plugins/Experimental/FastGeoStreaming/Source/FastGeoStreaming/Private/FastGeoWorldSubsystem.h
@@ -13,7 +13,7 @@ class IWorldPartitionHLODObject;
 enum class ELevelStreamingState : uint8;
 
 UCLASS()
-class FASTGEOSTREAMING_API UFastGeoWorldSubsystem : public UWorldSubsystem
+class FASTGEOSTREAMING_API UFastGeoWorldSubsystem : public UTickableWorldSubsystem
 {
 	GENERATED_BODY()
 
@@ -21,14 +21,19 @@ class FASTGEOSTREAMING_API UFastGeoWorldSubsystem : public UWorldSubsystem
 	~UFastGeoWorldSubsystem();
 
 public:
-	//~Begin WorldSubsystem interface
-	virtual void PostInitialize() override;
+	//~Begin USubsystem interface
 	virtual void Deinitialize() override;
-	//~End WorldSubsystem interface
+	//~End USubsystem interface
 
-	//~ Begin UWorldSubsystem Interface.
+	//~ Begin UWorldSubsystem interface.
+	virtual void PostInitialize() override;
 	virtual bool DoesSupportWorldType(const EWorldType::Type WorldType) const override;
-	//~ End UWorldSubsystem Interface.
+	//~ End UWorldSubsystem interface.
+
+	//~ Begin UTickableWorldSubsystem interface
+	virtual void Tick(float DeltaTime) override;
+	virtual TStatId GetStatId() const override;
+	//~ End UTickableWorldSubsystem interface
 
 	void AddToComponentsPendingRecreate(FFastGeoPrimitiveComponent* InComponentPendingRecreate);
 	void RemoveFromComponentsPendingRecreate(FFastGeoPrimitiveComponent* InComponentPendingRecreate);
@@ -36,6 +41,8 @@ public:
 
 	static bool IsEnableDebugView();
 
+	bool IsWaitingForCompletion() const;
+
 private:
 	void OnUpdateLevelStreaming();
 	void OnLevelStreamingStateChanged(UWorld* World, const ULevelStreaming* StreamingLevel, ULevel* LevelIfLoaded, ELevelStreamingState PrevState, ELevelStreamingState NewState);
@@ -53,6 +60,9 @@ private:
 
 	void ForEachHLODObjectInCell(const UWorldPartitionRuntimeCell*, TFunction<void(IWorldPartitionHLODObject*)>);
 
+	void RequestAsyncRenderStateTasksBudget_Concurrent(float& OutAvailableTimeBudgetMS, int32& OutAvaiableComponentsBudget, int32& OutTimeEpoch);
+	void CommitAsyncRenderStateTasksBudget_Concurrent(float InUsedTimeBudgetMS, int32& InUsedComponentsBudget, int32 TimeEpoch);
+
 	FDelegateHandle Handle_OnLevelStreamingStateChanged;
 	FDelegateHandle Handle_OnLevelBeginAddToWorld;
 	FDelegateHandle Handle_OnLevelBeginRemoveFromWorld;
@@ -60,4 +70,12 @@ private:
 
 	TArray<FWeakFastGeoComponent> ComponentsPendingRecreate;
 	static bool bEnableDebugView;
+
+	FRWLock Lock;
+	int32 TimeEpoch = 0;
+	float UsedAsyncRenderStateTasksTimeBudgetMS = 0;
+	int32 UsedNumComponentsToProcessBudget = 0;
+	bool bWaitingForCompletion = false;
+
+	friend class UFastGeoContainer;
 };
