FastGeo - Corrected visibility updates issues
* Visibility updates could fail if the async render state creation wasn't completed
* Ensure that failed updates are requeued from the GT
* Ensure UpdatePrimitivesDrawnInGame_RenderThread() is removing the proxies from lumen & raytracing scene representations, if necessary

diff --git a/Engine/Plugins/Experimental/FastGeoStreaming/Source/FastGeoStreaming/Private/FastGeoComponentCluster.cpp b/Engine/Plugins/Experimental/FastGeoStreaming/Source/FastGeoStreaming/Private/FastGeoComponentCluster.cpp
--- a/Engine/Plugins/Experimental/FastGeoStreaming/Source/FastGeoStreaming/Private/FastGeoComponentCluster.cpp
+++ b/Engine/Plugins/Experimental/FastGeoStreaming/Source/FastGeoStreaming/Private/FastGeoComponentCluster.cpp
@@ -3,6 +3,8 @@
 #include "FastGeoComponentCluster.h"
 #include "FastGeoContainer.h"
 #include "SceneInterface.h"
+#include "PrimitiveSceneInfo.h"
+#include "FastGeoWeakElement.h"
 
 const FFastGeoElementType FFastGeoComponentCluster::Type(&IFastGeoElement::Type);
 
@@ -98,9 +100,11 @@ ULevel* FFastGeoComponentCluster::GetLevel() const
 
 void FFastGeoComponentCluster::UpdateVisibility()
 {
-	TArray<FPrimitiveSceneProxy*, TInlineAllocator<32>> ShowProxies;
-	TArray<FPrimitiveSceneProxy*, TInlineAllocator<32>> HideProxies;
-	ForEachComponent<FFastGeoPrimitiveComponent>([&ShowProxies, &HideProxies](FFastGeoPrimitiveComponent& Component)
+	TRACE_CPUPROFILER_EVENT_SCOPE(FFastGeoComponentCluster::UpdateVisibility);
+
+	TArray<FFastGeoPrimitiveComponent*> ShowComponents;
+	TArray<FFastGeoPrimitiveComponent*> HideComponents;
+	ForEachComponent<FFastGeoPrimitiveComponent>([&ShowComponents, &HideComponents](FFastGeoPrimitiveComponent& Component)
 	{
 		const bool bOldIsDrawnInGame = Component.IsDrawnInGame();
 		Component.UpdateVisibility();
@@ -111,28 +115,94 @@ void FFastGeoComponentCluster::UpdateVisibility()
 			{
 				if (bIsDrawnInGame)
 				{
-					ShowProxies.Add(SceneProxy);
+					ShowComponents.Add(&Component);
 				}
 				else
 				{
-					HideProxies.Add(SceneProxy);
+					HideComponents.Add(&Component);
 				}
 			}
 		}
 	});
 
-	if (!ShowProxies.IsEmpty() || !HideProxies.IsEmpty())
+	UpdateVisibility_Internal(MoveTemp(ShowComponents), MoveTemp(HideComponents));
+}
+
+void FFastGeoComponentCluster::ForceUpdateVisibility(const TArray<FFastGeoPrimitiveComponent*>& Components, int32 UpdateCounter)
+{
+	TRACE_CPUPROFILER_EVENT_SCOPE(FFastGeoComponentCluster::ForceUpdateVisibility);
+
+	TArray<FFastGeoPrimitiveComponent*> ShowComponents;
+	TArray<FFastGeoPrimitiveComponent*> HideComponents;
+	for (FFastGeoPrimitiveComponent* Component : Components)
+	{
+		FWriteScopeLock WriteLock(Component->Lock.Get());
+		Component->UpdateVisibility();
+		const bool bIsDrawnInGame = Component->IsDrawnInGame();
+		if (FPrimitiveSceneProxy* SceneProxy = Component->GetSceneProxy())
+		{
+			if (bIsDrawnInGame)
+			{
+				ShowComponents.Add(Component);
+			}
+			else
+			{
+				HideComponents.Add(Component);
+			}
+		}
+	}
+
+	UpdateVisibility_Internal(MoveTemp(ShowComponents), MoveTemp(HideComponents), ++UpdateCounter);
+}
+
+void FFastGeoComponentCluster::UpdateVisibility_Internal(TArray<FFastGeoPrimitiveComponent*>&& ShowComponents, TArray<FFastGeoPrimitiveComponent*>&& HideComponents, int32 UpdateCounter)
+{
+	if (!ShowComponents.IsEmpty() || !HideComponents.IsEmpty())
 	{
-		ENQUEUE_RENDER_COMMAND()([ShowProxies = MoveTemp(ShowProxies), HideProxies = MoveTemp(HideProxies)](FRHICommandListBase&) mutable
+		FWeakFastGeoComponentCluster ClusterWeak(this);
+		ENQUEUE_RENDER_COMMAND()([ClusterWeak, UpdateCounter, ShowComponents = MoveTemp(ShowComponents), HideComponents = MoveTemp(HideComponents)](FRHICommandListBase&) mutable
 		{
-			if (!ShowProxies.IsEmpty())
+			TRACE_CPUPROFILER_EVENT_SCOPE(FFastGeoComponentCluster::UpdateVisibility_RenderThread);
+
+			if (!ClusterWeak.Get())
 			{
-				ShowProxies[0]->GetScene().UpdatePrimitivesDrawnInGame_RenderThread(ShowProxies, true);
+				return;
 			}
 
-			if (!HideProxies.IsEmpty())
+			TArray<FFastGeoPrimitiveComponent*> NotReadyComponents;
+			auto ProcessComponents = [&NotReadyComponents](TArray<FFastGeoPrimitiveComponent*>& InOutComponents, bool bShow)
 			{
-				HideProxies[0]->GetScene().UpdatePrimitivesDrawnInGame_RenderThread(HideProxies, false);
+				for (FFastGeoPrimitiveComponent* Component : InOutComponents)
+				{
+					FReadScopeLock ReadLock(Component->Lock.Get());
+					if (FPrimitiveSceneProxy* Proxy = Component->GetSceneProxy())
+					{
+						// Test whether the primitive was added to the scene (or is pending)
+						FPrimitiveSceneInfo* PrimitiveSceneInfo = Proxy->GetPrimitiveSceneInfo();
+						if (!PrimitiveSceneInfo->IsIndexValid())
+						{
+							NotReadyComponents.Add(Component);
+						}
+						else
+						{
+							Proxy->GetScene().UpdatePrimitivesDrawnInGame_RenderThread(MakeArrayView(&Proxy, 1), bShow);
+						}
+					}
+				}
+			};
+
+			ProcessComponents(ShowComponents, true);
+			ProcessComponents(HideComponents, false);
+
+			if (!NotReadyComponents.IsEmpty())
+			{
+				UE::Tasks::Launch(TEXT("ForceUpdateVisibility"), [ClusterWeak, UpdateCounter, NotReadyComponents = MoveTemp(NotReadyComponents)]()
+				{
+					if (FFastGeoComponentCluster* Cluster = ClusterWeak.Get())
+					{
+						Cluster->ForceUpdateVisibility(NotReadyComponents, UpdateCounter);
+					}
+				}, LowLevelTasks::ETaskPriority::Normal, UE::Tasks::EExtendedTaskPriority::GameThreadNormalPri);
 			}
 		});
 	}
diff --git a/Engine/Plugins/Experimental/FastGeoStreaming/Source/FastGeoStreaming/Private/FastGeoComponentCluster.h b/Engine/Plugins/Experimental/FastGeoStreaming/Source/FastGeoStreaming/Private/FastGeoComponentCluster.h
--- a/Engine/Plugins/Experimental/FastGeoStreaming/Source/FastGeoStreaming/Private/FastGeoComponentCluster.h
+++ b/Engine/Plugins/Experimental/FastGeoStreaming/Source/FastGeoStreaming/Private/FastGeoComponentCluster.h
@@ -161,6 +161,8 @@ protected:
 	}
 
 private:
+	void ForceUpdateVisibility(const TArray<FFastGeoPrimitiveComponent*>& Components, int32 UpdateCounter);
+	void UpdateVisibility_Internal(TArray<FFastGeoPrimitiveComponent*>&& ShowComponents, TArray<FFastGeoPrimitiveComponent*>&& HideComponents, int32 UpdateCounter = 0);
 	void InitializeDynamicProperties();
 	void SetOwnerContainer(UFastGeoContainer* InOwner);
 	void SetComponentClusterIndex(int32 InComponentClusterIndex);
diff --git a/Engine/Plugins/Experimental/FastGeoStreaming/Source/FastGeoStreaming/Private/FastGeoPrimitiveComponent.cpp b/Engine/Plugins/Experimental/FastGeoStreaming/Source/FastGeoStreaming/Private/FastGeoPrimitiveComponent.cpp
--- a/Engine/Plugins/Experimental/FastGeoStreaming/Source/FastGeoStreaming/Private/FastGeoPrimitiveComponent.cpp
+++ b/Engine/Plugins/Experimental/FastGeoStreaming/Source/FastGeoStreaming/Private/FastGeoPrimitiveComponent.cpp
@@ -312,6 +312,7 @@ FPrimitiveSceneDesc FFastGeoPrimitiveComponent::BuildSceneDesc()
 
 void FFastGeoPrimitiveComponent::CreateRenderState(FRegisterComponentContext* Context)
 {
+	FWriteScopeLock WriteLock(Lock.Get());
 	ProxyState = EProxyCreationState::Creating;
 	bRenderStateDirty = false;
 
@@ -374,6 +375,7 @@ void FFastGeoPrimitiveComponent::FFastGeoDestroyRenderStateContext::DestroyProxy
 
 void FFastGeoPrimitiveComponent::DestroyRenderState(FFastGeoDestroyRenderStateContext* Context)
 {
+	FWriteScopeLock WriteLock(Lock.Get());
 	if (GetSceneProxy())
 	{
 		check(ProxyState == EProxyCreationState::Created);
@@ -1038,4 +1040,4 @@ HHitProxy* UFastGeoPrimitiveComponentEditorProxy::CreatePrimitiveHitProxies(TArr
 	return nullptr;
 }
 //~ End IPrimitiveComponent interface
-#endif
\ No newline at end of file
+#endif
diff --git a/Engine/Plugins/Experimental/FastGeoStreaming/Source/FastGeoStreaming/Private/FastGeoPrimitiveComponent.h b/Engine/Plugins/Experimental/FastGeoStreaming/Source/FastGeoStreaming/Private/FastGeoPrimitiveComponent.h
--- a/Engine/Plugins/Experimental/FastGeoStreaming/Source/FastGeoStreaming/Private/FastGeoPrimitiveComponent.h
+++ b/Engine/Plugins/Experimental/FastGeoStreaming/Source/FastGeoStreaming/Private/FastGeoPrimitiveComponent.h
@@ -13,6 +13,7 @@
 #include "PhysicsEngine/PhysicsBodyInstanceOwnerInterface.h"
 #include "Chaos/ChaosUserEntity.h"
 #include "PSOPrecacheFwd.h"
+#include "Templates/DontCopy.h"
 #include "FastGeoPrimitiveComponent.generated.h"
 
 class FPrimitiveSceneProxy;
@@ -233,6 +234,9 @@ protected:
 	static_assert((int)EPSOPrecachePriority::Highest < 1 << 2);
 #endif
 
+	TDontCopy<FRWLock> Lock;
+
+	friend class FFastGeoComponentCluster;
 	friend class FPrimitiveComponentHelper;
 	friend class FFastGeoPhysicsBodyInstanceOwner;
 	friend class UFastGeoContainer;
diff --git a/Engine/Source/Runtime/Renderer/Private/RendererScene.cpp b/Engine/Source/Runtime/Renderer/Private/RendererScene.cpp
--- a/Engine/Source/Runtime/Renderer/Private/RendererScene.cpp
+++ b/Engine/Source/Runtime/Renderer/Private/RendererScene.cpp
@@ -1793,17 +1793,36 @@ void FScene::UpdatePrimitivesDrawnInGame_RenderThread(TArrayView<FPrimitiveScene
 	{
 		if (PrimitiveSceneProxy->IsDrawnInGame() != bDrawnInGame)
 		{
-			if (!bDrawnInGame)
+			FPrimitiveSceneInfo* PrimitiveSceneInfo = PrimitiveSceneProxy->GetPrimitiveSceneInfo();
+			const int32 PrimitiveIndex = PrimitiveSceneInfo->GetIndex();
+			const bool bPrimitiveIndexValid = PrimitiveSceneInfo->IsIndexValid();
+			check(bPrimitiveIndexValid);
+
+			if (!bDrawnInGame && bPrimitiveIndexValid)
 			{
-				DistanceFieldSceneData.RemovePrimitive(PrimitiveSceneProxy->GetPrimitiveSceneInfo());
+				DistanceFieldSceneData.RemovePrimitive(PrimitiveSceneInfo);
+				LumenRemovePrimitive(PrimitiveSceneInfo, PrimitiveIndex);
 			}
 
 			PrimitiveSceneProxy->SetDrawnInGame_RenderThread(bDrawnInGame);
 
-			if (bDrawnInGame)
+			if (bDrawnInGame && bPrimitiveIndexValid)
 			{
-				DistanceFieldSceneData.AddPrimitive(PrimitiveSceneProxy->GetPrimitiveSceneInfo());
+				DistanceFieldSceneData.AddPrimitive(PrimitiveSceneInfo);
+				LumenAddPrimitive(PrimitiveSceneInfo);
 			}
+				
+#if RHI_RAYTRACING
+			if (bPrimitiveIndexValid && PrimitiveSceneProxy->HasRayTracingRepresentation())
+			{
+				FRayTracingInstance CachedRayTracingInstance;
+				ERayTracingPrimitiveFlags& Flags = PrimitiveRayTracingFlags[PrimitiveIndex];
+
+				// Write flags
+				Flags = PrimitiveSceneInfo->Proxy->GetCachedRayTracingInstance(CachedRayTracingInstance);
+				FPrimitiveSceneInfo::UpdateCachedRayTracingInstance(PrimitiveSceneInfo, CachedRayTracingInstance, Flags);
+			}
+#endif
 		}
 	}
 }
