FastGeo - Readded lock that was lost in integration conflict

diff --git a/Engine/Plugins/Experimental/FastGeoStreaming/Source/FastGeoStreaming/Private/FastGeoComponentCluster.cpp b/Engine/Plugins/Experimental/FastGeoStreaming/Source/FastGeoStreaming/Private/FastGeoComponentCluster.cpp
--- a/Engine/Plugins/Experimental/FastGeoStreaming/Source/FastGeoStreaming/Private/FastGeoComponentCluster.cpp
+++ b/Engine/Plugins/Experimental/FastGeoStreaming/Source/FastGeoStreaming/Private/FastGeoComponentCluster.cpp
@@ -106,6 +106,7 @@ void FFastGeoComponentCluster::UpdateVisibility()
 	TArray<FFastGeoPrimitiveComponent*> HideComponents;
 	ForEachComponent<FFastGeoPrimitiveComponent>([&ShowComponents, &HideComponents](FFastGeoPrimitiveComponent& Component)
 	{
+		FWriteScopeLock WriteLock(Component.Lock.Get());
 		const bool bOldIsDrawnInGame = Component.IsDrawnInGame();
 		Component.UpdateVisibility();
 		const bool bIsDrawnInGame = Component.IsDrawnInGame();
