FastGeo - Fixes required to support dedicated servers

diff --git a/Engine/Plugins/Experimental/FastGeoStreaming/Source/FastGeoStreaming/Private/FastGeoContainer.cpp b/Engine/Plugins/Experimental/FastGeoStreaming/Source/FastGeoStreaming/Private/FastGeoContainer.cpp
--- a/Engine/Plugins/Experimental/FastGeoStreaming/Source/FastGeoStreaming/Private/FastGeoContainer.cpp
+++ b/Engine/Plugins/Experimental/FastGeoStreaming/Source/FastGeoStreaming/Private/FastGeoContainer.cpp
@@ -242,7 +242,10 @@ void UFastGeoContainer::Unregister()
 
 			ComponentCluster.ForEachComponent<FFastGeoPrimitiveComponent>([this](FFastGeoPrimitiveComponent& Component)
 			{
-				PendingDestroy.RenderState.ComponentsToProcess.Add(&Component);
+				if (Component.IsRenderStateCreated())
+				{
+					PendingDestroy.RenderState.ComponentsToProcess.Add(&Component);
+				}
 			});
 		});
 
diff --git a/Engine/Plugins/Experimental/FastGeoStreaming/Source/FastGeoStreaming/Private/FastGeoPrimitiveComponent.cpp b/Engine/Plugins/Experimental/FastGeoStreaming/Source/FastGeoStreaming/Private/FastGeoPrimitiveComponent.cpp
--- a/Engine/Plugins/Experimental/FastGeoStreaming/Source/FastGeoStreaming/Private/FastGeoPrimitiveComponent.cpp
+++ b/Engine/Plugins/Experimental/FastGeoStreaming/Source/FastGeoStreaming/Private/FastGeoPrimitiveComponent.cpp
@@ -505,6 +505,11 @@ bool FFastGeoPrimitiveComponent::IsRenderStateDirty() const
 
 bool FFastGeoPrimitiveComponent::ShouldCreateRenderState() const
 {
+	if (!FApp::CanEverRender())
+	{
+		return false;
+	}
+
 	// If the detail mode setting allows it, add it to the scene.
 	const bool bDetailModeAllowsRendering = DetailMode <= GetCachedScalabilityCVars().DetailMode;
 	if (!bDetailModeAllowsRendering)
