[FastGeo] Removed FastGeoStreamingManager and moved logic into FastGeoWorldSubsystem to have a context per world. This fixes issues in PIE when client & server run in the same process. Also fixed UFastGeoContainer::Tick not propagating bWaitForCompletion to the async physics state jobs.
#rnx

diff --git a/Engine/Plugins/Experimental/FastGeoStreaming/Source/FastGeoStreaming/Private/FastGeoContainer.cpp b/Engine/Plugins/Experimental/FastGeoStreaming/Source/FastGeoStreaming/Private/FastGeoContainer.cpp
--- a/Engine/Plugins/Experimental/FastGeoStreaming/Source/FastGeoStreaming/Private/FastGeoContainer.cpp
+++ b/Engine/Plugins/Experimental/FastGeoStreaming/Source/FastGeoStreaming/Private/FastGeoContainer.cpp
@@ -4,7 +4,6 @@
 #include "FastGeoComponentCluster.h"
 #include "FastGeoHLOD.h"
 #include "FastGeoComponent.h"
-#include "FastGeoStreamingManager.h"
 #include "FastGeoWorldSubsystem.h"
 #include "FastGeoLog.h"
 #include "AI/Navigation/NavigationElement.h"
@@ -199,12 +198,13 @@ void UFastGeoContainer::Register()
 			});
 		});
 
+		UFastGeoWorldSubsystem* WorldSubsystem = GetWorld()->GetSubsystem<UFastGeoWorldSubsystem>();
 		if (!PendingCreate.RenderState.ComponentsToProcess.IsEmpty())
 		{
-			FFastGeoStreamingManager::PushAsyncCreateRenderStateJob(this);
+			WorldSubsystem->PushAsyncCreateRenderStateJob(this);
 		}
 
-		FFastGeoStreamingManager::PushAsyncCreatePhysicsStateJobs(this);
+		WorldSubsystem->PushAsyncCreatePhysicsStateJobs(this);
 		RegisterToNavigationSystem();
 
 		bIsRegistered = true;
@@ -251,12 +251,13 @@ void UFastGeoContainer::Unregister()
 
 		UnregisterFromNavigationSystem();
 		
+		UFastGeoWorldSubsystem* WorldSubsystem = GetWorld()->GetSubsystem<UFastGeoWorldSubsystem>();
 		if (!PendingDestroy.RenderState.ComponentsToProcess.IsEmpty())
 		{
-			FFastGeoStreamingManager::PushAsyncDestroyRenderStateJob(this);
+			WorldSubsystem->PushAsyncDestroyRenderStateJob(this);
 		}
 
-		FFastGeoStreamingManager::PushAsyncDestroyPhysicsStateJobs(this);
+		WorldSubsystem->PushAsyncDestroyPhysicsStateJobs(this);
 		bIsRegistered = false;
 	}
 
@@ -268,14 +269,17 @@ void UFastGeoContainer::Unregister()
 
 void UFastGeoContainer::Tick(bool bWaitForCompletion)
 {
+	UWorld* World = GetWorld();
+	check(World);
+	UFastGeoWorldSubsystem* WorldSubsystem = World->GetSubsystem<UFastGeoWorldSubsystem>();
+
 	do
 	{
-		FFastGeoStreamingManager::ProcessAsyncRenderStateJobs(bWaitForCompletion);
-		UWorld* World = GetWorld();
-		check(World);
+		WorldSubsystem->ProcessAsyncRenderStateJobs(bWaitForCompletion);
+
 		if (FPhysScene* PhysScene = World->GetPhysicsScene())
 		{
-			PhysScene->ProcessAsyncPhysicsStateJobs();
+			PhysScene->ProcessAsyncPhysicsStateJobs(bWaitForCompletion);
 		}
 	} while (bWaitForCompletion && HasAnyPendingTasks());
 }
@@ -486,7 +490,8 @@ void UFastGeoContainer::OnCreateRenderStateEnd_GameThread()
 
 	if (PendingCreate.RenderState.TotalNumProcessed < PendingCreate.RenderState.ComponentsToProcess.Num())
 	{
-		FFastGeoStreamingManager::PushAsyncCreateRenderStateJob(this);
+		UFastGeoWorldSubsystem* WorldSubsystem = GetWorld()->GetSubsystem<UFastGeoWorldSubsystem>();
+		WorldSubsystem->PushAsyncCreateRenderStateJob(this);
 	}
 }
 
@@ -502,7 +507,8 @@ void UFastGeoContainer::OnDestroyRenderStateEnd_GameThread()
 
 	if (PendingDestroy.RenderState.TotalNumProcessed < PendingDestroy.RenderState.ComponentsToProcess.Num())
 	{
-		FFastGeoStreamingManager::PushAsyncDestroyRenderStateJob(this);
+		UFastGeoWorldSubsystem* WorldSubsystem = GetWorld()->GetSubsystem<UFastGeoWorldSubsystem>();
+		WorldSubsystem->PushAsyncDestroyRenderStateJob(this);
 	}
 }
 
diff --git a/Engine/Plugins/Experimental/FastGeoStreaming/Source/FastGeoStreaming/Private/FastGeoStreamingManager.cpp b/Engine/Plugins/Experimental/FastGeoStreaming/Source/FastGeoStreaming/Private/FastGeoStreamingManager.cpp
deleted file mode 100644
diff --git a/Engine/Plugins/Experimental/FastGeoStreaming/Source/FastGeoStreaming/Private/FastGeoStreamingManager.h b/Engine/Plugins/Experimental/FastGeoStreaming/Source/FastGeoStreaming/Private/FastGeoStreamingManager.h
deleted file mode 100644
diff --git a/Engine/Plugins/Experimental/FastGeoStreaming/Source/FastGeoStreaming/Private/FastGeoWorldSubsystem.cpp b/Engine/Plugins/Experimental/FastGeoStreaming/Source/FastGeoStreaming/Private/FastGeoWorldSubsystem.cpp
--- a/Engine/Plugins/Experimental/FastGeoStreaming/Source/FastGeoStreaming/Private/FastGeoWorldSubsystem.cpp
+++ b/Engine/Plugins/Experimental/FastGeoStreaming/Source/FastGeoStreaming/Private/FastGeoWorldSubsystem.cpp
@@ -1,7 +1,6 @@
 // Copyright Epic Games, Inc. All Rights Reserved.
 
 #include "FastGeoWorldSubsystem.h"
-#include "FastGeoStreamingManager.h"
 #include "FastGeoContainer.h"
 #include "FastGeoPrimitiveComponent.h"
 #include "Physics/Experimental/PhysScene_Chaos.h"
@@ -306,7 +305,7 @@ void UFastGeoWorldSubsystem::OnUpdateLevelStreaming()
 		TRACE_CPUPROFILER_EVENT_SCOPE(UFastGeoWorldSubsystem::OnUpdateLevelStreaming);
 
 		const bool bWaitForCompletion = true;
-		FFastGeoStreamingManager::ProcessAsyncRenderStateJobs(bWaitForCompletion);
+		ProcessAsyncRenderStateJobs(bWaitForCompletion);
 		if (FPhysScene* PhysScene = World->GetPhysicsScene())
 		{
 			PhysScene->ProcessAsyncPhysicsStateJobs(bWaitForCompletion);
@@ -475,4 +474,49 @@ void UFastGeoWorldSubsystem::CommitAsyncRenderStateTasksBudget_Concurrent(float
 	}
 }
 
+void UFastGeoWorldSubsystem::PushAsyncCreateRenderStateJob(UFastGeoContainer* FastGeo)
+{
+	check(IsInGameThread());
+	if (!AsyncRenderStateJobQueue.IsValid())
+	{
+		AsyncRenderStateJobQueue = MakeUnique<FFastGeoAsyncRenderStateJobQueue>(FastGeo->GetWorld()->Scene);
+	}
+	AsyncRenderStateJobQueue->AddJob({ FastGeo, FFastGeoAsyncRenderStateJobQueue::EJobType::CreateRenderState });
+}
+
+void UFastGeoWorldSubsystem::PushAsyncDestroyRenderStateJob(UFastGeoContainer* FastGeo)
+{
+	check(IsInGameThread());
+	if (!AsyncRenderStateJobQueue.IsValid())
+	{
+		AsyncRenderStateJobQueue = MakeUnique<FFastGeoAsyncRenderStateJobQueue>(FastGeo->GetWorld()->Scene);
+	}
+	AsyncRenderStateJobQueue->AddJob({ FastGeo, FFastGeoAsyncRenderStateJobQueue::EJobType::DestroyRenderState });
+}
+
+void UFastGeoWorldSubsystem::ProcessAsyncRenderStateJobs(bool bWaitForCompletion)
+{
+	check(IsInGameThread());
+	if (AsyncRenderStateJobQueue.IsValid())
+	{
+		AsyncRenderStateJobQueue->Tick(bWaitForCompletion);
+		if (AsyncRenderStateJobQueue->IsCompleted())
+		{
+			AsyncRenderStateJobQueue.Reset();
+		}
+	}
+}
+
+void UFastGeoWorldSubsystem::PushAsyncCreatePhysicsStateJobs(UFastGeoContainer* FastGeo)
+{
+	check(FPhysScene::SupportsAsyncPhysicsStateCreation());
+	FastGeo->OnCreatePhysicsStateBegin_GameThread();
+}
+
+void UFastGeoWorldSubsystem::PushAsyncDestroyPhysicsStateJobs(UFastGeoContainer* FastGeo)
+{
+	check(FPhysScene::SupportsAsyncPhysicsStateDestruction());
+	FastGeo->OnDestroyPhysicsStateBegin_GameThread();
+}
+
 #undef LOCTEXT_NAMESPACE
\ No newline at end of file
diff --git a/Engine/Plugins/Experimental/FastGeoStreaming/Source/FastGeoStreaming/Private/FastGeoWorldSubsystem.h b/Engine/Plugins/Experimental/FastGeoStreaming/Source/FastGeoStreaming/Private/FastGeoWorldSubsystem.h
--- a/Engine/Plugins/Experimental/FastGeoStreaming/Source/FastGeoStreaming/Private/FastGeoWorldSubsystem.h
+++ b/Engine/Plugins/Experimental/FastGeoStreaming/Source/FastGeoStreaming/Private/FastGeoWorldSubsystem.h
@@ -4,6 +4,7 @@
 
 #include "Subsystems/WorldSubsystem.h"
 #include "FastGeoWeakElement.h"
+#include "FastGeoAsyncRenderStateJobQueue.h"
 #include "FastGeoWorldSubsystem.generated.h"
 
 class FFastGeoPrimitiveComponent;
@@ -39,6 +40,36 @@ public:
 	void RemoveFromComponentsPendingRecreate(FFastGeoPrimitiveComponent* InComponentPendingRecreate);
 	void ProcessPendingRecreate();
 
+	/**
+	 * Start a job that will create render state asynchronously for a specified FastGeo container.
+	 * @param FastGeo: FastGeo container.
+	 */
+	void PushAsyncCreateRenderStateJob(UFastGeoContainer* FastGeo);
+
+	/**
+	 * Start a job that will destroy render state asynchronously for a specified FastGeo container.
+	 * @param FastGeo: FastGeo container.
+	 */
+	void PushAsyncDestroyRenderStateJob(UFastGeoContainer* FastGeo);
+
+	/**
+	 * Update progress of asynchronous render state creation and destruction.
+	 * @param bWaitForCompletion : Whether to wait for pending asynchronous tasks to complete.
+	 */
+	void ProcessAsyncRenderStateJobs(bool bWaitForCompletion = false);
+
+	/**
+	 * Start a job that will create physics state asynchronously for a specified FastGeo container.
+	 * @param FastGeo: FastGeo container.
+	 */
+	void PushAsyncCreatePhysicsStateJobs(UFastGeoContainer* FastGeo);
+
+	/**
+	 * Start a job that will destroy physics state asynchronously for a specified FastGeo container.
+	 * @param FastGeo: FastGeo container.
+	 */
+	void PushAsyncDestroyPhysicsStateJobs(UFastGeoContainer* FastGeo);
+
 	static bool IsEnableDebugView();
 
 	bool IsWaitingForCompletion() const;
@@ -77,5 +108,7 @@ private:
 	int32 UsedNumComponentsToProcessBudget = 0;
 	bool bWaitingForCompletion = false;
 
+	TUniquePtr<FFastGeoAsyncRenderStateJobQueue> AsyncRenderStateJobQueue;
+
 	friend class UFastGeoContainer;
 };
