FastGeo - Clear array view pointing to the original component MaterialStreamingRelativeBoxes array
* This is editor only data, doesn't need to be kept around for FastGeo

diff --git a/Engine/Plugins/Experimental/FastGeoStreaming/Source/FastGeoStreaming/Private/FastGeoStaticMeshComponent.cpp b/Engine/Plugins/Experimental/FastGeoStreaming/Source/FastGeoStreaming/Private/FastGeoStaticMeshComponent.cpp
--- a/Engine/Plugins/Experimental/FastGeoStreaming/Source/FastGeoStreaming/Private/FastGeoStaticMeshComponent.cpp
+++ b/Engine/Plugins/Experimental/FastGeoStreaming/Source/FastGeoStreaming/Private/FastGeoStaticMeshComponent.cpp
@@ -171,6 +171,10 @@ void FFastGeoStaticMeshComponentBase::ResetSceneProxyDescUnsupportedProperties()
 	SceneProxyDesc.SelectedEditorMaterial = INDEX_NONE;
 	SceneProxyDesc.SelectedEditorSection = INDEX_NONE;
 
+#if WITH_EDITORONLY_DATA
+	SceneProxyDesc.MaterialStreamingRelativeBoxes = TArrayView<uint32>();
+#endif
+
 	// Properties that will be initialized by InitializeSceneProxyDescDynamicProperties
 	SceneProxyDesc.NaniteResources = nullptr;
 	SceneProxyDesc.BodySetup = nullptr;
