[FastGeo] Fixed FFastGeoStaticMeshComponent not initializing its BodyInstance's ExternalCollisionProfileBodySetup when using default collision (in a cooked build).
This makes FBodyInstance::GetCollisionProfileName() to not always return the right value which can make FBodyInstance::LoadProfileData do the wrong thing and we end up with the wrong collision profile.

diff --git a/Engine/Plugins/Experimental/FastGeoStreaming/Source/FastGeoStreaming/Private/FastGeoStaticMeshComponent.cpp b/Engine/Plugins/Experimental/FastGeoStreaming/Source/FastGeoStreaming/Private/FastGeoStaticMeshComponent.cpp
--- a/Engine/Plugins/Experimental/FastGeoStreaming/Source/FastGeoStreaming/Private/FastGeoStaticMeshComponent.cpp
+++ b/Engine/Plugins/Experimental/FastGeoStreaming/Source/FastGeoStreaming/Private/FastGeoStaticMeshComponent.cpp
@@ -53,6 +53,23 @@ void FFastGeoStaticMeshComponentBase::Serialize(FArchive& Ar)
 	FArchive_Serialize_BitfieldBool(Ar, SceneProxyDesc.bDisallowNanite);
 	FArchive_Serialize_BitfieldBool(Ar, SceneProxyDesc.bForceDisableNanite);
 	FArchive_Serialize_BitfieldBool(Ar, SceneProxyDesc.bForceNaniteForMasked);
+	Ar << bUseDefaultCollision;
+}
+
+void FFastGeoStaticMeshComponentBase::InitializeDynamicProperties()
+{
+#if !WITH_EDITOR
+	// When using default collision, use the same collision profile as the StaticMesh
+	if (bUseDefaultCollision)
+	{
+		if (UBodySetup* BodySetup = GetBodySetup())
+		{
+			BodyInstance.UseExternalCollisionProfile(BodySetup);
+		}
+	}
+#endif
+
+	Super::InitializeDynamicProperties();
 }
 
 const Nanite::FResources* FFastGeoStaticMeshComponentBase::GetNaniteResources() const
@@ -147,6 +164,7 @@ void FFastGeoStaticMeshComponentBase::InitializeFromComponent(UActorComponent* C
 	Super::InitializeFromComponent(Component);
 
 	UStaticMeshComponent* StaticMeshComponent = CastChecked<UStaticMeshComponent>(Component);
+	bUseDefaultCollision = StaticMeshComponent->bUseDefaultCollision;
 	LocalBounds = GetStaticMesh()->GetBounds();
 	WorldBounds = LocalBounds.TransformBy(WorldTransform);
 }
diff --git a/Engine/Plugins/Experimental/FastGeoStreaming/Source/FastGeoStreaming/Private/FastGeoStaticMeshComponent.h b/Engine/Plugins/Experimental/FastGeoStreaming/Source/FastGeoStreaming/Private/FastGeoStaticMeshComponent.h
--- a/Engine/Plugins/Experimental/FastGeoStreaming/Source/FastGeoStreaming/Private/FastGeoStaticMeshComponent.h
+++ b/Engine/Plugins/Experimental/FastGeoStreaming/Source/FastGeoStreaming/Private/FastGeoStaticMeshComponent.h
@@ -53,6 +53,7 @@ public:
 
 	//~ Begin FFastGeoComponent interface
 	virtual void Serialize(FArchive& Ar) override;
+	virtual void InitializeDynamicProperties() override;
 	virtual UBodySetup* GetBodySetup() const override;
 #if WITH_EDITOR
 	virtual void InitializeFromComponent(UActorComponent* Component) override;
@@ -99,6 +100,7 @@ protected:
 #endif
 
 protected:
+	bool bUseDefaultCollision = false;
 	// Temporary until LODData support
 	struct FFastGeoStaticMeshComponentLODInfo
 	{
