FastGeo - Non-spatially loaded actors assigned to runtime data layers are rejected by the FastGeo transformer
* All actors that end up in WP cells at runtime should be candidates to be transformed by FastGeo. This includes non-spatially loaded actors that are assigned to runtime data layers.

diff --git a/Engine/Plugins/Experimental/FastGeoStreaming/Source/FastGeoStreaming/Private/FastGeoWorldPartitionRuntimeCellTransformer.cpp b/Engine/Plugins/Experimental/FastGeoStreaming/Source/FastGeoStreaming/Private/FastGeoWorldPartitionRuntimeCellTransformer.cpp
--- a/Engine/Plugins/Experimental/FastGeoStreaming/Source/FastGeoStreaming/Private/FastGeoWorldPartitionRuntimeCellTransformer.cpp
+++ b/Engine/Plugins/Experimental/FastGeoStreaming/Source/FastGeoStreaming/Private/FastGeoWorldPartitionRuntimeCellTransformer.cpp
@@ -658,9 +658,18 @@ FFastGeoTransformResult UFastGeoWorldPartitionRuntimeCellTransformer::CanTransfo
 		return { EFastGeoTransform::Reject, [InActor, Reason] { return FString::Printf(TEXT("Actor %s [%s]"), *InActor->GetName(), *Reason); } };
 	}
 
-	if (!InActor->GetIsSpatiallyLoaded())
+	// Reject actors that are never going to be streamed
+	// This test is only useful for the FastGeo debug mode, as normally during PIE & cook, only streamed actors will be processed by the cell transformer.
+	if (!InActor->GetLevel()->GetWorldPartitionRuntimeCell())
 	{
-		return { EFastGeoTransform::Reject, [InActor] { return FString::Printf(TEXT("Actor %s is non-spatially loaded"), *InActor->GetName()); } };
+		if (!InActor->GetIsSpatiallyLoaded())
+		{
+			bool bHasRuntimeDataLayers = Algo::AnyOf(InActor->GetDataLayerInstances(), [](const UDataLayerInstance* DataLayer) { return DataLayer->IsRuntime(); });
+			if (!bHasRuntimeDataLayers)
+			{
+				return { EFastGeoTransform::Reject, [InActor] { return FString::Printf(TEXT("Actor %s is non-spatially loaded"), *InActor->GetName()); } };
+			}
+		}
 	}
 
 	if (InActor->GetIsReplicated())
