[CellTransformer] Improve handling of components that can be ignored
* Add config settings to allow specifying component classes that can be ignored - IgnoredComponentClasses to ignore the class and subclasses or IgnoredExactComponentClasses to ignore just the class
* Add LevelInstanceComponent to IgnoredComponentClasses and SceneComponent to IgnoredExactComponentClasses
* Add CanIgnoreComponent API method to be used by transformers to find out if component can be ignored
* Add ForEachIgnoredComponentClass and ForEachIgnoredExactComponentClass API methods that can be overridden by transformers to override or extend the list of component classes ignored by default
* Refactor ISM CT and FastGeo CT to use the new APIs

This fixes the issue with ISM CellTransformer not removing Packed Level Actors properly

diff --git a/Engine/Config/BaseEngine.ini b/Engine/Config/BaseEngine.ini
--- a/Engine/Config/BaseEngine.ini
+++ b/Engine/Config/BaseEngine.ini
@@ -3886,6 +3886,10 @@ SendOSData=true
 EditorHashDefaultClass=/Script/CoreUObject.Class'/Script/Engine.WorldPartitionEditorSpatialHash'
 RuntimeHashDefaultClass=/Script/CoreUObject.Class'/Script/Engine.WorldPartitionRuntimeHashSet'
 
+[/Script/Engine.WorldPartitionRuntimeCellTransformerSettings]
++IgnoredComponentClasses=/Script/CoreUObject.Class'/Script/Engine.LevelInstanceComponent'
++IgnoredExactComponentClasses=/Script/CoreUObject.Class'/Script/Engine.SceneComponent'
+
 [BuildPatchTool]
 VerificationThreadCount=4
 
diff --git a/Engine/Plugins/Experimental/FastGeoStreaming/Config/FastGeoStreaming.ini b/Engine/Plugins/Experimental/FastGeoStreaming/Config/FastGeoStreaming.ini
--- a/Engine/Plugins/Experimental/FastGeoStreaming/Config/FastGeoStreaming.ini
+++ b/Engine/Plugins/Experimental/FastGeoStreaming/Config/FastGeoStreaming.ini
@@ -11,7 +11,5 @@
 +BuiltinDisallowedComponentClasses=/Script/CoreUObject.Class'/Script/Landscape.LandscapeMeshProxyComponent'
 +BuiltinDisallowedComponentClasses=/Script/CoreUObject.Class'/Script/Water.WaterBodyMeshComponent'
 +BuiltinDisallowedComponentClasses=/Script/CoreUObject.Class'/Script/Engine.SplineMeshComponent'
-+BuiltinIgnoredRemainingComponentClasses=/Script/CoreUObject.Class'/Script/Engine.LevelInstanceComponent'
 +BuiltinIgnoredRemainingComponentClasses=/Script/CoreUObject.Class'/Script/Engine.ActorTextureStreamingBuildDataComponent'
-+BuiltinIgnoredRemainingExactComponentClasses=/Script/CoreUObject.Class'/Script/Engine.SceneComponent'
 
diff --git a/Engine/Plugins/Experimental/FastGeoStreaming/Source/FastGeoStreaming/Private/FastGeoWorldPartitionRuntimeCellTransformer.cpp b/Engine/Plugins/Experimental/FastGeoStreaming/Source/FastGeoStreaming/Private/FastGeoWorldPartitionRuntimeCellTransformer.cpp
--- a/Engine/Plugins/Experimental/FastGeoStreaming/Source/FastGeoStreaming/Private/FastGeoWorldPartitionRuntimeCellTransformer.cpp
+++ b/Engine/Plugins/Experimental/FastGeoStreaming/Source/FastGeoStreaming/Private/FastGeoWorldPartitionRuntimeCellTransformer.cpp
@@ -353,6 +353,48 @@ void UFastGeoWorldPartitionRuntimeCellTransformer::Transform(ULevel* InLevel)
 	}
 }
 
+void UFastGeoWorldPartitionRuntimeCellTransformer::ForEachIgnoredComponentClass(TFunctionRef<bool(const TSubclassOf<UActorComponent>&)> Func) const
+{
+	Super::ForEachIgnoredComponentClass(Func);
+
+	for (const TSubclassOf<UActorComponent>& IgnoredComponentClass : IgnoredRemainingComponentClasses)
+	{
+		if (!Func(IgnoredComponentClass))
+		{
+			return;
+		}
+	}
+
+	for (const TSubclassOf<UActorComponent>& IgnoredComponentClass : BuiltinIgnoredRemainingComponentClasses)
+	{
+		if (!Func(IgnoredComponentClass))
+		{
+			return;
+		}
+	}
+}
+
+void UFastGeoWorldPartitionRuntimeCellTransformer::ForEachIgnoredExactComponentClass(TFunctionRef<bool(const TSubclassOf<UActorComponent>&)> Func) const
+{
+	Super::ForEachIgnoredExactComponentClass(Func);
+
+	for (const TSubclassOf<UActorComponent>& IgnoredComponentClass : IgnoredRemainingExactComponentClasses)
+	{
+		if (!Func(IgnoredComponentClass))
+		{
+			return;
+		}
+	}
+
+	for (const TSubclassOf<UActorComponent>& IgnoredComponentClass : BuiltinIgnoredRemainingExactComponentClasses)
+	{
+		if (!Func(IgnoredComponentClass))
+		{
+			return;
+		}
+	}
+}
+
 TMap<AActor*, TArray<AActor*>> UFastGeoWorldPartitionRuntimeCellTransformer::BuildActorsReferencesMap(const TArray<AActor*>& InActors)
 {
 	TRACE_CPUPROFILER_EVENT_SCOPE(UFastGeoWorldPartitionRuntimeCellTransformer::BuildActorsReferencesMap);
@@ -927,37 +969,9 @@ FFastGeoTransformResult UFastGeoWorldPartitionRuntimeCellTransformer::CanTransfo
 		return { EFastGeoTransform::Discard, [InComponent] { return FString::Printf(TEXT("Component %s has no collision and is not visible"), *FastGeo::GetComponentShortName(InComponent)); } };
 	}
 
-	UClass* ComponentClass = InComponent->GetClass();
-	for (const TSubclassOf<UActorComponent>& IgnoredComponentClass : IgnoredRemainingComponentClasses)
-	{
-		if (ComponentClass->IsChildOf(IgnoredComponentClass))
-		{
-			return { EFastGeoTransform::Discard };
-		}
-	}
-
-	for (const TSubclassOf<UActorComponent>& IgnoredComponentClass : BuiltinIgnoredRemainingComponentClasses)
-	{
-		if (ComponentClass->IsChildOf(IgnoredComponentClass))
-		{
-			return { EFastGeoTransform::Discard };
-		}
-	}
-
-	for (const TSubclassOf<UActorComponent>& IgnoredComponentClass : IgnoredRemainingExactComponentClasses)
+	if (CanIgnoreComponent(InComponent))
 	{
-		if (ComponentClass == IgnoredComponentClass)
-		{
-			return { EFastGeoTransform::Discard };
-		}
-	}
-
-	for (const TSubclassOf<UActorComponent>& IgnoredComponentClass : BuiltinIgnoredRemainingExactComponentClasses)
-	{
-		if (ComponentClass == IgnoredComponentClass)
-		{
-			return { EFastGeoTransform::Discard };
-		}
+		return { EFastGeoTransform::Discard };
 	}
 
 	return EFastGeoTransform::Allow;
diff --git a/Engine/Plugins/Experimental/FastGeoStreaming/Source/FastGeoStreaming/Public/FastGeoWorldPartitionRuntimeCellTransformer.h b/Engine/Plugins/Experimental/FastGeoStreaming/Source/FastGeoStreaming/Public/FastGeoWorldPartitionRuntimeCellTransformer.h
--- a/Engine/Plugins/Experimental/FastGeoStreaming/Source/FastGeoStreaming/Public/FastGeoWorldPartitionRuntimeCellTransformer.h
+++ b/Engine/Plugins/Experimental/FastGeoStreaming/Source/FastGeoStreaming/Public/FastGeoWorldPartitionRuntimeCellTransformer.h
@@ -54,6 +54,8 @@ public:
 public:
 	//~ Begin UWorldPartitionRuntimeCellTransformer Interface.
 	virtual void Transform(ULevel* InLevel) override final;
+	virtual void ForEachIgnoredComponentClass(TFunctionRef<bool(const TSubclassOf<UActorComponent>&)> Func) const override final;
+	virtual void ForEachIgnoredExactComponentClass(TFunctionRef<bool(const TSubclassOf<UActorComponent>&)> Func) const override final;
 	//~ End UWorldPartitionRuntimeCellTransformer Interface.
 
 	//~ Begin UObject Interface.
diff --git a/Engine/Source/Runtime/Engine/Private/WorldPartition/WorldPartitionRuntimeCellTransformer.cpp b/Engine/Source/Runtime/Engine/Private/WorldPartition/WorldPartitionRuntimeCellTransformer.cpp
--- a/Engine/Source/Runtime/Engine/Private/WorldPartition/WorldPartitionRuntimeCellTransformer.cpp
+++ b/Engine/Source/Runtime/Engine/Private/WorldPartition/WorldPartitionRuntimeCellTransformer.cpp
@@ -6,4 +6,61 @@ UWorldPartitionRuntimeCellTransformer::UWorldPartitionRuntimeCellTransformer(con
 	: Super(ObjectInitializer)
 {}
 
+#if WITH_EDITOR
+void UWorldPartitionRuntimeCellTransformer::ForEachIgnoredComponentClass(TFunctionRef<bool(const TSubclassOf<UActorComponent>&)> Func) const
+{
+	for (const TSubclassOf<UActorComponent>& IgnoredComponentClass : GetDefault<UWorldPartitionRuntimeCellTransformerSettings>()->IgnoredComponentClasses)
+	{
+		if (!Func(IgnoredComponentClass))
+		{
+			return;
+		}
+	}
+}
+
+void UWorldPartitionRuntimeCellTransformer::ForEachIgnoredExactComponentClass(TFunctionRef<bool(const TSubclassOf<UActorComponent>&)> Func) const
+{
+	for (const TSubclassOf<UActorComponent>& IgnoredExactComponentClass : GetDefault<UWorldPartitionRuntimeCellTransformerSettings>()->IgnoredExactComponentClasses)
+	{
+		if (!Func(IgnoredExactComponentClass))
+		{
+			return;
+		}
+	}
+}
+
+bool UWorldPartitionRuntimeCellTransformer::CanIgnoreComponent(const UActorComponent* InComponent) const
+{
+	bool bCanIgnore = false;
+	const UClass* ComponentClass = InComponent->GetClass();
+
+	ForEachIgnoredComponentClass([&bCanIgnore, &ComponentClass](const TSubclassOf<UActorComponent>& IgnoredComponentClass)
+	{
+		if (ComponentClass->IsChildOf(IgnoredComponentClass))
+		{
+			bCanIgnore = true;
+			return false;
+		}
+		return true;
+	});
+
+	if (bCanIgnore)
+	{
+		return true;
+	}
+
+	ForEachIgnoredExactComponentClass([&bCanIgnore, &ComponentClass](const TSubclassOf<UActorComponent>& IgnoredExactComponentClass)
+	{
+		if (ComponentClass == IgnoredExactComponentClass)
+		{
+			bCanIgnore = true;
+			return false;
+		}
+		return true;
+	});
+
+	return bCanIgnore;
+}
+#endif
+
 const FName UWorldPartitionRuntimeCellTransformer::NAME_CellTransformerIgnoreActor(TEXT("CellTransformer_IgnoreActor"));
diff --git a/Engine/Source/Runtime/Engine/Private/WorldPartition/WorldPartitionRuntimeCellTransformerISM.cpp b/Engine/Source/Runtime/Engine/Private/WorldPartition/WorldPartitionRuntimeCellTransformerISM.cpp
--- a/Engine/Source/Runtime/Engine/Private/WorldPartition/WorldPartitionRuntimeCellTransformerISM.cpp
+++ b/Engine/Source/Runtime/Engine/Private/WorldPartition/WorldPartitionRuntimeCellTransformerISM.cpp
@@ -195,29 +195,15 @@ bool UWorldPartitionRuntimeCellTransformerISM::CanAutoInstanceActor(AActor* InAc
 
 bool UWorldPartitionRuntimeCellTransformerISM::CanRemoveActor(AActor* InActor) const
 {
-	uint32 NumRuntimeComponents = 0;
 	for (UActorComponent* Component : InActor->GetComponents())
 	{
-		if (!Component->IsEditorOnly())
+		if (!Component->IsEditorOnly() && !CanIgnoreComponent(Component))
 		{
-			NumRuntimeComponents++;
+			return false;
 		}
 	}
 		
-	if (!NumRuntimeComponents)
-	{
-		return true;
-	}
-
-	if (NumRuntimeComponents == 1)
-	{
-		if (InActor->GetRootComponent() && ExactCast<USceneComponent>(InActor->GetRootComponent()))
-		{
-			return true;
-		}
-	}
-
-	return false;
+	return true;
 }
 #endif
 
diff --git a/Engine/Source/Runtime/Engine/Public/WorldPartition/WorldPartitionRuntimeCellTransformer.h b/Engine/Source/Runtime/Engine/Public/WorldPartition/WorldPartitionRuntimeCellTransformer.h
--- a/Engine/Source/Runtime/Engine/Public/WorldPartition/WorldPartitionRuntimeCellTransformer.h
+++ b/Engine/Source/Runtime/Engine/Public/WorldPartition/WorldPartitionRuntimeCellTransformer.h
@@ -2,8 +2,24 @@
 #pragma once
 
 #include "CoreMinimal.h"
+#include "Components/ActorComponent.h"
 #include "WorldPartitionRuntimeCellTransformer.generated.h"
 
+UCLASS(MinimalAPI, Config=Engine)
+class UWorldPartitionRuntimeCellTransformerSettings : public UObject
+{
+	GENERATED_BODY()
+
+public:
+#if WITH_EDITORONLY_DATA
+	UPROPERTY(Config)
+	TArray<TSubclassOf<UActorComponent>> IgnoredComponentClasses;
+
+	UPROPERTY(Config)
+	TArray<TSubclassOf<UActorComponent>> IgnoredExactComponentClasses;
+#endif
+};
+
 UCLASS(MinimalAPI, Abstract)
 class UWorldPartitionRuntimeCellTransformer : public UObject
 {
@@ -14,9 +30,16 @@ public:
 	virtual void PreTransform(ULevel* InLevel) {}
 	virtual void Transform(ULevel* InLevel) {}
 	virtual void PostTransform(ULevel* InLevel) {}
+	ENGINE_API virtual void ForEachIgnoredComponentClass(TFunctionRef<bool(const TSubclassOf<UActorComponent>&)> Func) const;
+	ENGINE_API virtual void ForEachIgnoredExactComponentClass(TFunctionRef<bool(const TSubclassOf<UActorComponent>&)> Func) const;
 #endif
 
 	bool IsEnabled() const { return bEnabled; }
+
+protected:
+#if WITH_EDITOR
+	ENGINE_API virtual bool CanIgnoreComponent(const UActorComponent* InComponent) const;
+#endif
 	
 protected:
 	// Tag used to force exclude actors from any cell transformation
