[FastGeo] Fixed regression, non-primitive components were ignored from the transformation process and could be lost if all actor's primitive component were transformable.

diff --git a/Engine/Plugins/Experimental/FastGeoStreaming/Source/FastGeoStreaming/Private/FastGeoWorldPartitionRuntimeCellTransformer.cpp b/Engine/Plugins/Experimental/FastGeoStreaming/Source/FastGeoStreaming/Private/FastGeoWorldPartitionRuntimeCellTransformer.cpp
--- a/Engine/Plugins/Experimental/FastGeoStreaming/Source/FastGeoStreaming/Private/FastGeoWorldPartitionRuntimeCellTransformer.cpp
+++ b/Engine/Plugins/Experimental/FastGeoStreaming/Source/FastGeoStreaming/Private/FastGeoWorldPartitionRuntimeCellTransformer.cpp
@@ -744,12 +744,62 @@ FFastGeoTransformResult UFastGeoWorldPartitionRuntimeCellTransformer::CanTransfo
 		return { EFastGeoTransform::Reject, [InActor] { return FString::Printf(TEXT("Actor %s is a Blueprint Actor with logic"), *InActor->GetName()); } };
 	}
 
+	auto CanIgnoreRemainingComponent = [this](UActorComponent* InComponent)
+	{
+		UClass* ComponentClass = InComponent->GetClass();
+		for (const TSubclassOf<UActorComponent>& IgnoredComponentClass : IgnoredRemainingComponentClasses)
+		{
+			if (ComponentClass->IsChildOf(IgnoredComponentClass))
+			{
+				return true;
+			}
+		}
+
+		for (const TSubclassOf<UActorComponent>& IgnoredComponentClass : BuiltinIgnoredRemainingComponentClasses)
+		{
+			if (ComponentClass->IsChildOf(IgnoredComponentClass))
+			{
+				return true;
+			}
+		}
+
+		for (const TSubclassOf<UActorComponent>& IgnoredComponentClass : IgnoredRemainingExactComponentClasses)
+		{
+			if (ComponentClass == IgnoredComponentClass)
+			{
+				return true;
+			}
+		}
+
+		for (const TSubclassOf<UActorComponent>& IgnoredComponentClass : BuiltinIgnoredRemainingExactComponentClasses)
+		{
+			if (ComponentClass == IgnoredComponentClass)
+			{
+				return true;
+			}
+		}
+		return false;
+	};
+
 	// Gather transformable components
 	TArray<UActorComponent*> TransformResults[EnumToIndex(EFastGeoTransform::MAX)];
-	InActor->ForEachComponent<UPrimitiveComponent>(false, [this, &TransformResults](UPrimitiveComponent* PrimitiveComponent)
+	InActor->ForEachComponent<UActorComponent>(false, [this, &TransformResults, CanIgnoreRemainingComponent](UActorComponent* Component)
 	{
-		FFastGeoTransformResult TransformComponentResult = CanTransformComponent(PrimitiveComponent);
-		TransformResults[TransformComponentResult.GetResultIndex()].Add(PrimitiveComponent);
+		if (UPrimitiveComponent* PrimitiveComponent = Cast<UPrimitiveComponent>(Component))
+		{
+			FFastGeoTransformResult TransformComponentResult = CanTransformComponent(PrimitiveComponent);
+			TransformResults[TransformComponentResult.GetResultIndex()].Add(Component);
+		}
+		else if (CanIgnoreRemainingComponent(Component))
+		{
+			FFastGeoTransformResult TransformComponentResult = { EFastGeoTransform::Discard };
+			TransformResults[TransformComponentResult.GetResultIndex()].Add(Component);
+		}
+		else
+		{
+			FFastGeoTransformResult TransformComponentResult = { EFastGeoTransform::Reject, [Component] { return FString::Printf(TEXT("Component %s class %s is not allowed"), *FastGeo::GetComponentShortName(Component), *Component->GetClass()->GetName()); } };
+			TransformResults[TransformComponentResult.GetResultIndex()].Add(Component);
+		}
 	});
 
 	if (TransformResults[EnumToIndex(EFastGeoTransform::Allow)].IsEmpty())
