[FastGeo] When the transformed component is the RootComponent, replace it with a SceneComponent and remap attachment of other components.
Without that fix, the cooking generates warnings about an invalid component and at runtime, non-transformed components are at the wrong location (see JIRA).

diff --git a/Engine/Plugins/Experimental/FastGeoStreaming/Source/FastGeoStreaming/Private/FastGeoWorldPartitionRuntimeCellTransformer.cpp b/Engine/Plugins/Experimental/FastGeoStreaming/Source/FastGeoStreaming/Private/FastGeoWorldPartitionRuntimeCellTransformer.cpp
--- a/Engine/Plugins/Experimental/FastGeoStreaming/Source/FastGeoStreaming/Private/FastGeoWorldPartitionRuntimeCellTransformer.cpp
+++ b/Engine/Plugins/Experimental/FastGeoStreaming/Source/FastGeoStreaming/Private/FastGeoWorldPartitionRuntimeCellTransformer.cpp
@@ -325,6 +325,21 @@ void UFastGeoWorldPartitionRuntimeCellTransformer::Transform(ULevel* InLevel)
 			{
 				InLevel->Actors[Entry.Value.ActorIndex] = nullptr;
 			}
+			else if (USceneComponent* OldRootComponent = Actor->GetRootComponent(); OldRootComponent && !IsValid(OldRootComponent))
+			{
+				// Replace removed root component with a SceneComponent and remap attachment of other components
+				USceneComponent* NewRootComponent = NewObject<USceneComponent>(Actor);
+				NewRootComponent->SetRelativeTransform(OldRootComponent->GetRelativeTransform());
+				Actor->SetRootComponent(NewRootComponent);
+
+				Actor->ForEachComponent<USceneComponent>(false, [OldRootComponent, NewRootComponent](USceneComponent* Component)
+				{
+					if (Component->GetAttachParent() == OldRootComponent)
+					{
+						Component->SetupAttachment(NewRootComponent);
+					}
+				});
+			}
 
 			if (FastGeoHLOD.IsValid())
 			{
