[FastGeo] Simplified fix for non-primitive components that were ignored from the transformation process and that could be lost if all actor's primitive component were transformable.

diff --git a/Engine/Plugins/Experimental/FastGeoStreaming/Source/FastGeoStreaming/Private/FastGeoWorldPartitionRuntimeCellTransformer.cpp b/Engine/Plugins/Experimental/FastGeoStreaming/Source/FastGeoStreaming/Private/FastGeoWorldPartitionRuntimeCellTransformer.cpp
--- a/Engine/Plugins/Experimental/FastGeoStreaming/Source/FastGeoStreaming/Private/FastGeoWorldPartitionRuntimeCellTransformer.cpp
+++ b/Engine/Plugins/Experimental/FastGeoStreaming/Source/FastGeoStreaming/Private/FastGeoWorldPartitionRuntimeCellTransformer.cpp
@@ -759,62 +759,12 @@ FFastGeoTransformResult UFastGeoWorldPartitionRuntimeCellTransformer::CanTransfo
 		return { EFastGeoTransform::Reject, [InActor] { return FString::Printf(TEXT("Actor %s is a Blueprint Actor with logic"), *InActor->GetName()); } };
 	}
 
-	auto CanIgnoreRemainingComponent = [this](UActorComponent* InComponent)
-	{
-		UClass* ComponentClass = InComponent->GetClass();
-		for (const TSubclassOf<UActorComponent>& IgnoredComponentClass : IgnoredRemainingComponentClasses)
-		{
-			if (ComponentClass->IsChildOf(IgnoredComponentClass))
-			{
-				return true;
-			}
-		}
-
-		for (const TSubclassOf<UActorComponent>& IgnoredComponentClass : BuiltinIgnoredRemainingComponentClasses)
-		{
-			if (ComponentClass->IsChildOf(IgnoredComponentClass))
-			{
-				return true;
-			}
-		}
-
-		for (const TSubclassOf<UActorComponent>& IgnoredComponentClass : IgnoredRemainingExactComponentClasses)
-		{
-			if (ComponentClass == IgnoredComponentClass)
-			{
-				return true;
-			}
-		}
-
-		for (const TSubclassOf<UActorComponent>& IgnoredComponentClass : BuiltinIgnoredRemainingExactComponentClasses)
-		{
-			if (ComponentClass == IgnoredComponentClass)
-			{
-				return true;
-			}
-		}
-		return false;
-	};
-
 	// Gather transformable components
 	TArray<UActorComponent*> TransformResults[EnumToIndex(EFastGeoTransform::MAX)];
-	InActor->ForEachComponent<UActorComponent>(false, [this, &TransformResults, CanIgnoreRemainingComponent](UActorComponent* Component)
+	InActor->ForEachComponent<UActorComponent>(false, [this, &TransformResults](UActorComponent* Component)
 	{
-		if (UPrimitiveComponent* PrimitiveComponent = Cast<UPrimitiveComponent>(Component))
-		{
-			FFastGeoTransformResult TransformComponentResult = CanTransformComponent(PrimitiveComponent);
-			TransformResults[TransformComponentResult.GetResultIndex()].Add(Component);
-		}
-		else if (CanIgnoreRemainingComponent(Component))
-		{
-			FFastGeoTransformResult TransformComponentResult = { EFastGeoTransform::Discard };
-			TransformResults[TransformComponentResult.GetResultIndex()].Add(Component);
-		}
-		else
-		{
-			FFastGeoTransformResult TransformComponentResult = { EFastGeoTransform::Reject, [Component] { return FString::Printf(TEXT("Component %s class %s is not allowed"), *FastGeo::GetComponentShortName(Component), *Component->GetClass()->GetName()); } };
-			TransformResults[TransformComponentResult.GetResultIndex()].Add(Component);
-		}
+		FFastGeoTransformResult TransformComponentResult = CanTransformComponent(Component);
+		TransformResults[TransformComponentResult.GetResultIndex()].Add(Component);
 	});
 
 	if (TransformResults[EnumToIndex(EFastGeoTransform::Allow)].IsEmpty())
@@ -904,8 +854,18 @@ FFastGeoTransformResult UFastGeoWorldPartitionRuntimeCellTransformer::IsAllowedC
 	return { EFastGeoTransform::Reject, [InComponent, ComponentClass] { return FString::Printf(TEXT("Component %s class is an unsupported class (%s)"), *FastGeo::GetComponentShortName(InComponent), *ComponentClass->GetName()); } };
 };
 
-FFastGeoTransformResult UFastGeoWorldPartitionRuntimeCellTransformer::CanTransformComponent(UPrimitiveComponent* InComponent) const
+FFastGeoTransformResult UFastGeoWorldPartitionRuntimeCellTransformer::CanTransformComponent(UActorComponent* InComponent) const
 {
+	UPrimitiveComponent* PrimitiveComponent = Cast<UPrimitiveComponent>(InComponent);
+	if (!PrimitiveComponent)
+	{
+		if (CanIgnoreComponent(InComponent))
+		{
+			return { EFastGeoTransform::Discard };
+		}
+		return { EFastGeoTransform::Reject, [PrimitiveComponent] { return FString::Printf(TEXT("Component %s class %s is not allowed"), *FastGeo::GetComponentShortName(PrimitiveComponent), *PrimitiveComponent->GetClass()->GetName()); } };
+	}
+
 	FFastGeoTransformResult AllowedComponentClassResult = IsAllowedComponentClass(InComponent);
 	if (AllowedComponentClassResult.GetResult() != EFastGeoTransform::Allow)
 	{
@@ -913,49 +873,49 @@ FFastGeoTransformResult UFastGeoWorldPartitionRuntimeCellTransformer::CanTransfo
 	}
 
 	FString Reason;
-	if (!IsComponentTransformable(InComponent, Reason))
+	if (!IsComponentTransformable(PrimitiveComponent, Reason))
 	{
-		return { EFastGeoTransform::Reject, [InComponent, Reason] { return FString::Printf(TEXT("Component %s [%s]"), *FastGeo::GetComponentShortName(InComponent), *Reason); } };
+		return { EFastGeoTransform::Reject, [PrimitiveComponent, Reason] { return FString::Printf(TEXT("Component %s [%s]"), *FastGeo::GetComponentShortName(PrimitiveComponent), *Reason); } };
 	}
 
-	if (InComponent->IsEditorOnly())
+	if (PrimitiveComponent->IsEditorOnly())
 	{
-		return { EFastGeoTransform::Discard, [InComponent] { return FString::Printf(TEXT("Component %s is editor-only"), *FastGeo::GetComponentShortName(InComponent)); } };
+		return { EFastGeoTransform::Discard, [PrimitiveComponent] { return FString::Printf(TEXT("Component %s is editor-only"), *FastGeo::GetComponentShortName(PrimitiveComponent)); } };
 	}
 
-	if (InComponent->GetLODParentPrimitive())
+	if (PrimitiveComponent->GetLODParentPrimitive())
 	{
-		return { EFastGeoTransform::Reject, [InComponent] { return FString::Printf(TEXT("Component %s has a valid LOD Parent Primitive"), *FastGeo::GetComponentShortName(InComponent)); } };
+		return { EFastGeoTransform::Reject, [PrimitiveComponent] { return FString::Printf(TEXT("Component %s has a valid LOD Parent Primitive"), *FastGeo::GetComponentShortName(PrimitiveComponent)); } };
 	}
 
-	if (InComponent->Mobility != EComponentMobility::Static)
+	if (PrimitiveComponent->Mobility != EComponentMobility::Static)
 	{
-		return { EFastGeoTransform::Reject, [InComponent] { return FString::Printf(TEXT("Component %s Mobility is not Static"), *FastGeo::GetComponentShortName(InComponent)); } };
+		return { EFastGeoTransform::Reject, [PrimitiveComponent] { return FString::Printf(TEXT("Component %s Mobility is not Static"), *FastGeo::GetComponentShortName(PrimitiveComponent)); } };
 	}
 
-	bool bIsComponentVisible = InComponent->IsVisible() && !InComponent->GetOwner()->IsHidden();
-	bool bShouldAddToRenderScene = bIsComponentVisible || InComponent->bCastHiddenShadow || InComponent->bAffectIndirectLightingWhileHidden || InComponent->bRayTracingFarField;
+	bool bIsComponentVisible = PrimitiveComponent->IsVisible() && !PrimitiveComponent->GetOwner()->IsHidden();
+	bool bShouldAddToRenderScene = bIsComponentVisible || PrimitiveComponent->bCastHiddenShadow || PrimitiveComponent->bAffectIndirectLightingWhileHidden || PrimitiveComponent->bRayTracingFarField;
 
-	if (UStaticMeshComponent* StaticMeshComponent = Cast<UStaticMeshComponent>(InComponent))
+	if (UStaticMeshComponent* StaticMeshComponent = Cast<UStaticMeshComponent>(PrimitiveComponent))
 	{
 		if (!StaticMeshComponent->GetStaticMesh())
 		{
-			return { EFastGeoTransform::Discard, [InComponent] { return FString::Printf(TEXT("Component %s has an invalid static mesh"), *FastGeo::GetComponentShortName(InComponent)); } };
+			return { EFastGeoTransform::Discard, [PrimitiveComponent] { return FString::Printf(TEXT("Component %s has an invalid static mesh"), *FastGeo::GetComponentShortName(PrimitiveComponent)); } };
 		}
 
 		// Make sure BodyInstance CollisionEnabled is updated first before testing below
 		StaticMeshComponent->UpdateCollisionFromStaticMesh();
 	}
 
-	if (UInstancedStaticMeshComponent* InstancedStaticMeshComponent = Cast<UInstancedStaticMeshComponent>(InComponent))
+	if (UInstancedStaticMeshComponent* InstancedStaticMeshComponent = Cast<UInstancedStaticMeshComponent>(PrimitiveComponent))
 	{
 		if (InstancedStaticMeshComponent->GetNumInstances() == 0)
 		{
-			return { EFastGeoTransform::Discard, [InComponent] { return FString::Printf(TEXT("Component %s has no instances"), *FastGeo::GetComponentShortName(InComponent)); } };
+			return { EFastGeoTransform::Discard, [PrimitiveComponent] { return FString::Printf(TEXT("Component %s has no instances"), *FastGeo::GetComponentShortName(PrimitiveComponent)); } };
 		}
 	}
 
-	if (UHierarchicalInstancedStaticMeshComponent* HierarchicalInstancedStaticMeshComponent = Cast<UHierarchicalInstancedStaticMeshComponent>(InComponent))
+	if (UHierarchicalInstancedStaticMeshComponent* HierarchicalInstancedStaticMeshComponent = Cast<UHierarchicalInstancedStaticMeshComponent>(PrimitiveComponent))
 	{
 		// FastGeo doesn't really support HISMC. These components get converted to ISMC.
 		// However, we can afford to convert nanite HISMC as all the LODing logic is performed by Nanite.
@@ -964,74 +924,74 @@ FFastGeoTransformResult UFastGeoWorldPartitionRuntimeCellTransformer::CanTransfo
 		{
 			if (HierarchicalInstancedStaticMeshComponent->GetStaticMesh()->GetNumLODs() > 1)
 			{
-				return { EFastGeoTransform::Reject, [InComponent] { return FString::Printf(TEXT("Hierarchical instanced static mesh component %s has multiple LODs"), *FastGeo::GetComponentShortName(InComponent)); } };
+				return { EFastGeoTransform::Reject, [PrimitiveComponent] { return FString::Printf(TEXT("Hierarchical instanced static mesh component %s has multiple LODs"), *FastGeo::GetComponentShortName(PrimitiveComponent)); } };
 			}
 		}
 	}
 
-	if (USkinnedMeshComponent* SkinnedMeshComponent = Cast<USkinnedMeshComponent>(InComponent))
+	if (USkinnedMeshComponent* SkinnedMeshComponent = Cast<USkinnedMeshComponent>(PrimitiveComponent))
 	{
 		if (!SkinnedMeshComponent->GetSkinnedAsset())
 		{
-			return { EFastGeoTransform::Discard, [InComponent] { return FString::Printf(TEXT("Skinned mesh component %s has an invalid skinned asset"), *FastGeo::GetComponentShortName(InComponent)); } };
+			return { EFastGeoTransform::Discard, [PrimitiveComponent] { return FString::Printf(TEXT("Skinned mesh component %s has an invalid skinned asset"), *FastGeo::GetComponentShortName(PrimitiveComponent)); } };
 		}
 
 		if (SkinnedMeshComponent->LeaderPoseComponent.IsValid())
 		{
-			return { EFastGeoTransform::Reject, [InComponent] { return FString::Printf(TEXT("Skinned mesh component %s has a leader pose component"), *FastGeo::GetComponentShortName(InComponent)); } };
+			return { EFastGeoTransform::Reject, [PrimitiveComponent] { return FString::Printf(TEXT("Skinned mesh component %s has a leader pose component"), *FastGeo::GetComponentShortName(PrimitiveComponent)); } };
 		}
 
 		if (SkinnedMeshComponent->GetCollisionEnabled() != ECollisionEnabled::NoCollision)
 		{
-			return { EFastGeoTransform::Reject, [InComponent] { return FString::Printf(TEXT("Skinned mesh component %s has collisions enabled"), *FastGeo::GetComponentShortName(InComponent)); } };
+			return { EFastGeoTransform::Reject, [PrimitiveComponent] { return FString::Printf(TEXT("Skinned mesh component %s has collisions enabled"), *FastGeo::GetComponentShortName(PrimitiveComponent)); } };
 		}
 
 		if (SkinnedMeshComponent->IsNavigationRelevant())
 		{
-			return { EFastGeoTransform::Reject, [InComponent] { return FString::Printf(TEXT("Skinned mesh component %s is navigation relevant"), *FastGeo::GetComponentShortName(InComponent)); } };
+			return { EFastGeoTransform::Reject, [PrimitiveComponent] { return FString::Printf(TEXT("Skinned mesh component %s is navigation relevant"), *FastGeo::GetComponentShortName(PrimitiveComponent)); } };
 		}
 
 		bShouldAddToRenderScene &= !SkinnedMeshComponent->bHideSkin;
 	}
 
-	if (USkeletalMeshComponent* SkeletalMeshComponent = Cast<USkeletalMeshComponent>(InComponent))
+	if (USkeletalMeshComponent* SkeletalMeshComponent = Cast<USkeletalMeshComponent>(PrimitiveComponent))
 	{
 		if ((SkeletalMeshComponent->GetAnimationMode() == EAnimationMode::AnimationSingleNode && SkeletalMeshComponent->AnimationData.AnimToPlay != nullptr) ||
 			(SkeletalMeshComponent->GetAnimationMode() == EAnimationMode::AnimationBlueprint && SkeletalMeshComponent->AnimClass != nullptr) ||
 			(SkeletalMeshComponent->GetAnimationMode() == EAnimationMode::AnimationCustomMode))
 		{
-			return { EFastGeoTransform::Reject, [InComponent] { return FString::Printf(TEXT("Skeletal mesh component %s is animated"), *FastGeo::GetComponentShortName(InComponent)); } };
+			return { EFastGeoTransform::Reject, [PrimitiveComponent] { return FString::Printf(TEXT("Skeletal mesh component %s is animated"), *FastGeo::GetComponentShortName(PrimitiveComponent)); } };
 		}
 	}
 
-	if (UInstancedSkinnedMeshComponent* InstancedSkinnedMeshComponent = Cast<UInstancedSkinnedMeshComponent>(InComponent))
+	if (UInstancedSkinnedMeshComponent* InstancedSkinnedMeshComponent = Cast<UInstancedSkinnedMeshComponent>(PrimitiveComponent))
 	{
 		for (const FAnimBankItem& AnimBankItem : InstancedSkinnedMeshComponent->GetAnimBankItems())
 		{
 			if (AnimBankItem.BankAsset != nullptr)
 			{
-				return { EFastGeoTransform::Reject, [InComponent] { return FString::Printf(TEXT("Instanced skinned mesh component %s is animated"), *FastGeo::GetComponentShortName(InComponent)); } };
+				return { EFastGeoTransform::Reject, [PrimitiveComponent] { return FString::Printf(TEXT("Instanced skinned mesh component %s is animated"), *FastGeo::GetComponentShortName(PrimitiveComponent)); } };
 			}
 		}
 
 		if (InstancedSkinnedMeshComponent->GetInstanceCount() == 0)
 		{
-			return { EFastGeoTransform::Discard, [InComponent] { return FString::Printf(TEXT("Instanced skinned mesh component %s has no instances"), *FastGeo::GetComponentShortName(InComponent)); } };
+			return { EFastGeoTransform::Discard, [PrimitiveComponent] { return FString::Printf(TEXT("Instanced skinned mesh component %s has no instances"), *FastGeo::GetComponentShortName(PrimitiveComponent)); } };
 		}
 	}
 
-	const bool bIsCollisionEnabled = FastGeo::IsCollisionEnabled(InComponent);
+	const bool bIsCollisionEnabled = FastGeo::IsCollisionEnabled(PrimitiveComponent);
 
 	// If collision is enabled, only allow if async physics state creation and destruction are supported 
-	if (bIsCollisionEnabled && (!InComponent->AllowsAsyncPhysicsStateCreation() || !InComponent->AllowsAsyncPhysicsStateDestruction()))
+	if (bIsCollisionEnabled && (!PrimitiveComponent->AllowsAsyncPhysicsStateCreation() || !PrimitiveComponent->AllowsAsyncPhysicsStateDestruction()))
 	{
-		return { EFastGeoTransform::Reject, [InComponent] { return FString::Printf(TEXT("Component %s has collision enabled but doesn't allow asynchronous physics state creation/destruction"), *FastGeo::GetComponentShortName(InComponent)); } };
+		return { EFastGeoTransform::Reject, [PrimitiveComponent] { return FString::Printf(TEXT("Component %s has collision enabled but doesn't allow asynchronous physics state creation/destruction"), *FastGeo::GetComponentShortName(PrimitiveComponent)); } };
 	}
 
 	// Disallow transform if collision is disabled and component doesn't need to be added to the render scene
 	if (!bIsCollisionEnabled && !bShouldAddToRenderScene)
 	{
-		return { EFastGeoTransform::Discard, [InComponent] { return FString::Printf(TEXT("Component %s has no collision and is not visible"), *FastGeo::GetComponentShortName(InComponent)); } };
+		return { EFastGeoTransform::Discard, [PrimitiveComponent] { return FString::Printf(TEXT("Component %s has no collision and is not visible"), *FastGeo::GetComponentShortName(PrimitiveComponent)); } };
 	}
 
 	if (CanIgnoreComponent(InComponent))
diff --git a/Engine/Plugins/Experimental/FastGeoStreaming/Source/FastGeoStreaming/Public/FastGeoWorldPartitionRuntimeCellTransformer.h b/Engine/Plugins/Experimental/FastGeoStreaming/Source/FastGeoStreaming/Public/FastGeoWorldPartitionRuntimeCellTransformer.h
--- a/Engine/Plugins/Experimental/FastGeoStreaming/Source/FastGeoStreaming/Public/FastGeoWorldPartitionRuntimeCellTransformer.h
+++ b/Engine/Plugins/Experimental/FastGeoStreaming/Source/FastGeoStreaming/Public/FastGeoWorldPartitionRuntimeCellTransformer.h
@@ -75,7 +75,7 @@ protected:
 	bool IsBlueprintActorWithLogic(AActor* InActor) const;
 	FFastGeoTransformResult CanTransformActor(AActor* InActor, bool& bOutIsActorFullyTransformable, TArray<UActorComponent*>& OutTransformableComponents) const;
 	FFastGeoTransformResult IsAllowedActorClass(AActor* InActor) const;
-	FFastGeoTransformResult CanTransformComponent(UPrimitiveComponent* InComponent) const;
+	FFastGeoTransformResult CanTransformComponent(UActorComponent* InComponent) const;
 	FFastGeoTransformResult IsAllowedComponentClass(UActorComponent* InComponent) const;
 	bool CanRemoveActor(AActor* InActor, TSet<UActorComponent*>& InIgnoredComponents) const;
 	bool CanAlwaysIgnoreActor(AActor* InActor) const;
