Fixed CIS - Static Analysis following CL 43915336

diff --git a/Engine/Plugins/Experimental/FastGeoStreaming/Source/FastGeoStreaming/Private/FastGeoWorldPartitionRuntimeCellTransformer.cpp b/Engine/Plugins/Experimental/FastGeoStreaming/Source/FastGeoStreaming/Private/FastGeoWorldPartitionRuntimeCellTransformer.cpp
--- a/Engine/Plugins/Experimental/FastGeoStreaming/Source/FastGeoStreaming/Private/FastGeoWorldPartitionRuntimeCellTransformer.cpp
+++ b/Engine/Plugins/Experimental/FastGeoStreaming/Source/FastGeoStreaming/Private/FastGeoWorldPartitionRuntimeCellTransformer.cpp
@@ -863,7 +863,7 @@ FFastGeoTransformResult UFastGeoWorldPartitionRuntimeCellTransformer::CanTransfo
 		{
 			return { EFastGeoTransform::Discard };
 		}
-		return { EFastGeoTransform::Reject, [PrimitiveComponent] { return FString::Printf(TEXT("Component %s class %s is not allowed"), *FastGeo::GetComponentShortName(PrimitiveComponent), *PrimitiveComponent->GetClass()->GetName()); } };
+		return { EFastGeoTransform::Reject, [InComponent] { return FString::Printf(TEXT("Component %s class %s is not allowed"), *FastGeo::GetComponentShortName(InComponent), *InComponent->GetClass()->GetName()); } };
 	}
 
 	FFastGeoTransformResult AllowedComponentClassResult = IsAllowedComponentClass(InComponent);
