Fixed CIS caused by CL 44762423.

diff --git a/Engine/Plugins/Experimental/FastGeoStreaming/Source/FastGeoStreaming/Private/FastGeoStreamingModule.cpp b/Engine/Plugins/Experimental/FastGeoStreaming/Source/FastGeoStreaming/Private/FastGeoStreamingModule.cpp
--- a/Engine/Plugins/Experimental/FastGeoStreaming/Source/FastGeoStreaming/Private/FastGeoStreamingModule.cpp
+++ b/Engine/Plugins/Experimental/FastGeoStreaming/Source/FastGeoStreaming/Private/FastGeoStreamingModule.cpp
@@ -6,6 +6,7 @@
 #include "FastGeoInstancedStaticMeshComponent.h"
 #include "FastGeoWorldSubsystem.h"
 #include "Modules/ModuleManager.h"
+#include "Engine/World.h"
 
 IMPLEMENT_MODULE(FFastGeoStreamingModule, FastGeoStreaming);
 
