[FastGeo] Fixed bug in FastGeo when FApp::CanEverRender returns false and the FastGeoContainer doesn't contain any physics.
- Added IAsyncPhysicsStateProcessor to OnAsyncCreatePhysicsStateBegin_GameThread().
- FastGeoContainer now changes its internal physics creation state to creating in OnAsyncCreatePhysicsStateBegin_GameThread which prevents UFastGeoWorldSubsystem::OnAddLevelToWorldExtension from returning that it has completed.
- Optim: Prevent pushing physics tasks when the FastGeoContainer CollisionComponents is empty.

diff --git a/Engine/Plugins/Experimental/FastGeoStreaming/Source/FastGeoStreaming/Private/FastGeoContainer.cpp b/Engine/Plugins/Experimental/FastGeoStreaming/Source/FastGeoStreaming/Private/FastGeoContainer.cpp
--- a/Engine/Plugins/Experimental/FastGeoStreaming/Source/FastGeoStreaming/Private/FastGeoContainer.cpp
+++ b/Engine/Plugins/Experimental/FastGeoStreaming/Source/FastGeoStreaming/Private/FastGeoContainer.cpp
@@ -737,28 +737,34 @@ void UFastGeoContainer::OnCreatePhysicsStateBegin_GameThread()
 {
 	TRACE_CPUPROFILER_EVENT_SCOPE(UFastGeoContainer::OnCreatePhysicsStateBegin_GameThread);
 
-	PendingCreate.PhysicsState.ComponentsToProcess = &CollisionComponents;
-	check(PendingCreate.PhysicsState.TotalNumProcessed == 0);
+	if (!CollisionComponents.IsEmpty())
+	{
+		PendingCreate.PhysicsState.ComponentsToProcess = &CollisionComponents;
+		check(PendingCreate.PhysicsState.TotalNumProcessed == 0);
 
-	UWorld* World = GetWorld();
-	check(World);
-	FPhysScene* PhysScene = World->GetPhysicsScene();
-	check(PhysScene);
-	PhysScene->PushAsyncCreatePhysicsStateJob(this);
+		UWorld* World = GetWorld();
+		check(World);
+		FPhysScene* PhysScene = World->GetPhysicsScene();
+		check(PhysScene);
+		PhysScene->PushAsyncCreatePhysicsStateJob(this);
+	}
 }
 
 void UFastGeoContainer::OnDestroyPhysicsStateBegin_GameThread()
 {
 	TRACE_CPUPROFILER_EVENT_SCOPE(UFastGeoContainer::OnDestroyPhysicsStateBegin_GameThread);
 
-	PendingDestroy.PhysicsState.ComponentsToProcess = &CollisionComponents;
-	check(PendingDestroy.PhysicsState.TotalNumProcessed == 0);
+	if (!CollisionComponents.IsEmpty())
+	{
+		PendingDestroy.PhysicsState.ComponentsToProcess = &CollisionComponents;
+		check(PendingDestroy.PhysicsState.TotalNumProcessed == 0);
 
-	UWorld* World = GetWorld();
-	check(World);
-	FPhysScene* PhysScene = World->GetPhysicsScene();
-	check(PhysScene);
-	verify(PhysScene->PushAsyncDestroyPhysicsStateJob(this));
+		UWorld* World = GetWorld();
+		check(World);
+		FPhysScene* PhysScene = World->GetPhysicsScene();
+		check(PhysScene);
+		verify(PhysScene->PushAsyncDestroyPhysicsStateJob(this));
+	}
 }
 
 bool UFastGeoContainer::AllowsAsyncPhysicsStateCreation() const
@@ -783,14 +789,19 @@ UObject* UFastGeoContainer::GetAsyncPhysicsStateObject() const
 	return const_cast<UFastGeoContainer*>(this);
 }
 
-bool UFastGeoContainer::OnAsyncCreatePhysicsState(const UE::FTimeout& Timeout)
+void UFastGeoContainer::OnAsyncCreatePhysicsStateBegin_GameThread()
 {
-	TRACE_CPUPROFILER_EVENT_SCOPE(UFastGeoContainer::OnAsyncCreatePhysicsState);
-	check((PhysicsStateCreation == EPhysicsStateCreation::NotCreated) || 
-		  ((PhysicsStateCreation == EPhysicsStateCreation::Creating) && (PendingCreate.PhysicsState.TotalNumProcessed > 0) && !PendingCreate.PhysicsState.IsCompleted()));
+	check(PhysicsStateCreation == EPhysicsStateCreation::NotCreated);
+	check(PendingCreate.PhysicsState.ComponentsToProcess);
+	check(PendingCreate.PhysicsState.ComponentsToProcess->Num());
 	PhysicsStateCreation = EPhysicsStateCreation::Creating;
+}
 
-	check(PendingCreate.PhysicsState.ComponentsToProcess);
+bool UFastGeoContainer::OnAsyncCreatePhysicsState(const UE::FTimeout& Timeout)
+{
+	TRACE_CPUPROFILER_EVENT_SCOPE(UFastGeoContainer::OnAsyncCreatePhysicsState);
+	check((PhysicsStateCreation == EPhysicsStateCreation::Creating) && !PendingCreate.PhysicsState.IsCompleted());
+	
 	for (int Index = PendingCreate.PhysicsState.TotalNumProcessed; Index < PendingCreate.PhysicsState.ComponentsToProcess->Num(); ++Index)
 	{
 		FFastGeoComponent* Component = (*PendingCreate.PhysicsState.ComponentsToProcess)[Index];
@@ -824,6 +835,8 @@ void UFastGeoContainer::OnAsyncDestroyPhysicsStateBegin_GameThread()
 {
 	TRACE_CPUPROFILER_EVENT_SCOPE(UFastGeoContainer::OnAsyncDestroyPhysicsStateBegin_GameThread);
 	check(PhysicsStateCreation == EPhysicsStateCreation::Created);
+	check(PendingDestroy.PhysicsState.ComponentsToProcess);
+	check(PendingDestroy.PhysicsState.ComponentsToProcess->Num());
 	PhysicsStateCreation = EPhysicsStateCreation::Destroying;
 
 	for (FFastGeoComponent* Component : CollisionComponents)
@@ -836,8 +849,6 @@ bool UFastGeoContainer::OnAsyncDestroyPhysicsState(const UE::FTimeout& Timeout)
 {
 	TRACE_CPUPROFILER_EVENT_SCOPE(UFastGeoContainer::OnAsyncDestroyPhysicsState);
 
-	//for (FFastGeoComponent* Component : CollisionComponents)
-	check(PendingDestroy.PhysicsState.ComponentsToProcess);
 	for (int Index = PendingDestroy.PhysicsState.TotalNumProcessed; Index < PendingDestroy.PhysicsState.ComponentsToProcess->Num(); ++Index)
 	{
 		FFastGeoComponent* Component = (*PendingDestroy.PhysicsState.ComponentsToProcess)[Index];
diff --git a/Engine/Plugins/Experimental/FastGeoStreaming/Source/FastGeoStreaming/Private/FastGeoContainer.h b/Engine/Plugins/Experimental/FastGeoStreaming/Source/FastGeoStreaming/Private/FastGeoContainer.h
--- a/Engine/Plugins/Experimental/FastGeoStreaming/Source/FastGeoStreaming/Private/FastGeoContainer.h
+++ b/Engine/Plugins/Experimental/FastGeoStreaming/Source/FastGeoStreaming/Private/FastGeoContainer.h
@@ -149,6 +149,7 @@ public:
 	virtual bool AllowsAsyncPhysicsStateDestruction() const override final;
 	virtual bool IsAsyncPhysicsStateCreated() const override final;
 	virtual UObject* GetAsyncPhysicsStateObject() const override final;
+	virtual void OnAsyncCreatePhysicsStateBegin_GameThread() override;
 	virtual bool OnAsyncCreatePhysicsState(const UE::FTimeout& Timeout) override;
 	virtual void OnAsyncCreatePhysicsStateEnd_GameThread() override;
 	virtual void OnAsyncDestroyPhysicsStateBegin_GameThread() override;
diff --git a/Engine/Source/Runtime/Engine/Private/PhysicsEngine/Experimental/PhysScene_AsyncPhysicsStateJobQueue.cpp b/Engine/Source/Runtime/Engine/Private/PhysicsEngine/Experimental/PhysScene_AsyncPhysicsStateJobQueue.cpp
--- a/Engine/Source/Runtime/Engine/Private/PhysicsEngine/Experimental/PhysScene_AsyncPhysicsStateJobQueue.cpp
+++ b/Engine/Source/Runtime/Engine/Private/PhysicsEngine/Experimental/PhysScene_AsyncPhysicsStateJobQueue.cpp
@@ -242,6 +242,11 @@ void FPhysScene_AsyncPhysicsStateJobQueue::FJob::OnPreExecute_GameThread() const
 	TRACE_CPUPROFILER_EVENT_SCOPE(FPhysScene_AsyncPhysicsStateJobQueue::OnPreExecute_GameThread);
 	switch (Type)
 	{
+		case EJobType::CreatePhysicsState:
+		{
+			Processor->OnAsyncCreatePhysicsStateBegin_GameThread();
+		}
+		break;
 		case EJobType::DestroyPhysicsState:
 		{
 			Processor->OnAsyncDestroyPhysicsStateBegin_GameThread();
diff --git a/Engine/Source/Runtime/Engine/Public/Physics/Experimental/AsyncPhysicsStateProcessorInterface.h b/Engine/Source/Runtime/Engine/Public/Physics/Experimental/AsyncPhysicsStateProcessorInterface.h
--- a/Engine/Source/Runtime/Engine/Public/Physics/Experimental/AsyncPhysicsStateProcessorInterface.h
+++ b/Engine/Source/Runtime/Engine/Public/Physics/Experimental/AsyncPhysicsStateProcessorInterface.h
@@ -28,6 +28,9 @@ public:
 	/** Returns body setups that needs to create their physics meshes before physics state asynchronous creation */
 	virtual void CollectBodySetupsWithPhysicsMeshesToCreate(TSet<UBodySetup*>& OutBodySetups) const {}
 
+	/** Called on the GameThread before the component's physic engine information is created. */
+	virtual void OnAsyncCreatePhysicsStateBegin_GameThread() {}
+
 	/** Used to create any physics engine information for this component outside of the GameThread. */
 	virtual bool OnAsyncCreatePhysicsState(const UE::FTimeout& TimeOut) { return true; }
 
