FastGeo: Move headers required for usage from PCG

#rnx

diff --git a/Engine/Plugins/Experimental/FastGeoStreaming/Source/FastGeoStreaming/Internal/FastGeoComponent.h b/Engine/Plugins/Experimental/FastGeoStreaming/Source/FastGeoStreaming/Internal/FastGeoComponent.h
new file mode 100644
--- /dev/null
+++ b/Engine/Plugins/Experimental/FastGeoStreaming/Source/FastGeoStreaming/Internal/FastGeoComponent.h
@@ -0,0 +1,130 @@
+// Copyright Epic Games, Inc. All Rights Reserved.
+
+#pragma once
+
+#include "CoreMinimal.h"
+#include "IFastGeoElement.h"
+#include "FastGeoElementType.h"
+#include "Physics/Experimental/AsyncPhysicsStateProcessorInterface.h"
+#include "FastGeoComponent.generated.h"
+
+class UWorld;
+class UBodySetup;
+class UActorComponent;
+class FFastGeoComponentCluster;
+class UFastGeoContainer;
+class FSceneInterface;
+class FPrimitiveSceneProxy;
+class FRegisterComponentContext;
+
+class FASTGEOSTREAMING_API FFastGeoComponent : public IFastGeoElement
+{
+public:
+	typedef IFastGeoElement Super;
+
+	/** Static type identifier for this element class */
+	static const FFastGeoElementType Type;
+
+	FFastGeoComponent(int32 InComponentIndex = INDEX_NONE, FFastGeoElementType InType = Type);
+	virtual ~FFastGeoComponent() = default;
+
+#if WITH_EDITOR
+	virtual void InitializeFromComponent(UActorComponent* Component);
+	virtual UClass* GetEditorProxyClass() const;
+#endif
+
+	int32 GetComponentIndex() const { return ComponentIndex; }
+	virtual UBodySetup* GetBodySetup() const { return nullptr; }
+	virtual bool IsCollisionEnabled() const { return false; }
+	virtual void Serialize(FArchive& Ar);
+	virtual void InitializeDynamicProperties() {}
+	virtual void OnAsyncCreatePhysicsState();
+	virtual void OnAsyncDestroyPhysicsStateBegin_GameThread();
+	virtual void OnAsyncDestroyPhysicsState();
+	virtual void OnAsyncDestroyPhysicsStateEnd_GameThread();
+	void OnAsyncCreatePhysicsStateEnd_GameThread();
+
+	FFastGeoComponentCluster* GetOwnerComponentCluster() const;
+	UFastGeoContainer* GetOwnerContainer() const;
+	UWorld* GetWorld() const;
+	bool IsRegistered() const;
+	FLinearColor GetDebugColor() const;
+
+#if WITH_EDITOR
+	void SetEditorProxy(UFastGeoComponentEditorProxy* InComponentEditorProxy)
+	{
+		ComponentEditorProxy = InComponentEditorProxy;
+	}
+		
+	template <typename TEditorProxy>
+	TEditorProxy* GetEditorProxy() const
+	{
+		return Cast<TEditorProxy>(ComponentEditorProxy);
+	}
+#endif
+
+protected:
+	// Persistent Data
+	int32 ComponentIndex;
+
+	// Transient Data
+	FFastGeoComponentCluster* Owner;
+
+	enum EPhysicsStateCreation
+	{
+		NotCreated,
+		Creating,
+		Created,
+		Destroying
+	};
+
+	EPhysicsStateCreation PhysicsStateCreation = EPhysicsStateCreation::NotCreated;
+
+private:
+	void SetOwnerComponentCluster(FFastGeoComponentCluster* InOwner);
+	friend class FFastGeoComponentCluster;
+
+	friend FArchive& operator<<(FArchive& Ar, FFastGeoComponent& Component);
+
+#if WITH_EDITOR
+	UFastGeoComponentEditorProxy* ComponentEditorProxy;
+#endif
+};
+
+
+UCLASS()
+class UFastGeoComponentEditorProxy : public UObject
+{
+	GENERATED_BODY()
+
+#if WITH_EDITOR
+	friend class UFastGeoContainer;
+
+protected:
+	template <typename TComponent>
+	const TComponent& GetComponent() const
+	{
+		return FastGeoComponent->CastToRef<TComponent>();
+	}
+
+	template <typename TComponent>
+	TComponent& GetComponent()
+	{
+		return FastGeoComponent->CastToRef<TComponent>();
+	}
+
+private:
+	void SetFastGeoComponent(FFastGeoComponent* InFastGeoComponent)
+	{
+		FastGeoComponent = InFastGeoComponent;
+	}
+
+	const FFastGeoComponent* GetFastGeoComponent() const
+	{
+		return FastGeoComponent;
+	}
+
+private:
+	FFastGeoComponent* FastGeoComponent;
+#endif
+};
diff --git a/Engine/Plugins/Experimental/FastGeoStreaming/Source/FastGeoStreaming/Internal/FastGeoComponentCluster.h b/Engine/Plugins/Experimental/FastGeoStreaming/Source/FastGeoStreaming/Internal/FastGeoComponentCluster.h
new file mode 100644
--- /dev/null
+++ b/Engine/Plugins/Experimental/FastGeoStreaming/Source/FastGeoStreaming/Internal/FastGeoComponentCluster.h
@@ -0,0 +1,172 @@
+// Copyright Epic Games, Inc. All Rights Reserved.
+
+#pragma once
+
+#include "IFastGeoElement.h"
+#include "FastGeoElementType.h"
+#include "FastGeoStaticMeshComponent.h"
+#include "FastGeoInstancedStaticMeshComponent.h"
+#include "FastGeoSkinnedMeshComponent.h"
+#include "FastGeoInstancedSkinnedMeshComponent.h"
+
+class ULevel;
+class UFastGeoContainer;
+class FFastGeoComponent;
+class FArchive;
+
+#if WITH_EDITOR
+#include "UObject/ObjectSaveContext.h"
+#endif // WITH_EDITOR
+
+class FFastGeoComponentCluster : public IFastGeoElement
+{
+public:
+	typedef IFastGeoElement Super;
+
+	/** Static type identifier for this element class */
+	static const FFastGeoElementType Type;
+
+	FFastGeoComponentCluster(UFastGeoContainer* InOwner = nullptr, FName InName = NAME_None, FFastGeoElementType InType = Type);
+
+#if WITH_EDITOR
+	FFastGeoComponentCluster(const FFastGeoComponentCluster& InOther);
+#else
+	FFastGeoComponentCluster(const FFastGeoComponentCluster& InOther) = delete;
+#endif
+
+	virtual ~FFastGeoComponentCluster() {}
+	virtual void OnRegister() {}
+	virtual void OnUnregister() {}
+	virtual void Serialize(FArchive& Ar);
+	virtual bool IsVisible() const { return true; }
+
+	UFastGeoContainer* GetOwnerContainer() const;
+	ULevel* GetLevel() const;
+	FFastGeoComponent* GetComponent(uint32 InComponentTypeID, int32 InComponentIndex);
+	const FString& GetName() const { return Name; }
+	int32 GetComponentClusterIndex() const { return ComponentClusterIndex; }
+	void UpdateVisibility();
+
+#if WITH_EDITOR
+	bool HasComponents() const;
+	virtual FFastGeoComponent& AddComponent(FFastGeoElementType InComponentType);
+	virtual void PreSave(FObjectPreSaveContext ObjectSaveContext) {}
+#endif
+
+	template <typename TComponent = const FFastGeoComponent, typename TFunc>
+	void ForEachComponent(TFunc&& InFunc) const
+	{
+		ForEachComponent<const FFastGeoComponentCluster, const TComponent, TFunc>(this, Forward<TFunc>(InFunc));
+	}
+
+	template <typename TComponent = FFastGeoComponent, typename TFunc>
+	void ForEachComponent(TFunc&& InFunc)
+	{
+		ForEachComponent<FFastGeoComponentCluster, TComponent, TFunc>(this, Forward<TFunc>(InFunc));
+	}
+
+	template <typename TComponent = const FFastGeoComponent, typename TFunc>
+	bool ForEachComponentBreakable(TFunc&& InFunc) const
+	{
+		return ForEachComponentBreakable<const FFastGeoComponentCluster, const TComponent, TFunc>(this, Forward<TFunc>(InFunc));
+	}
+
+	template <typename TComponent = FFastGeoComponent, typename TFunc>
+	bool ForEachComponentBreakable(TFunc&& InFunc)
+	{
+		return ForEachComponentBreakable<FFastGeoComponentCluster, TComponent, TFunc>(this, Forward<TFunc>(InFunc));
+	}
+
+protected:
+	// Transient Data
+	UFastGeoContainer* Owner;
+
+	// Persistent Data
+	FString Name;
+	int32 ComponentClusterIndex;
+
+	TArray<FFastGeoStaticMeshComponent> StaticMeshComponents;
+	TArray<FFastGeoInstancedStaticMeshComponent> InstancedStaticMeshComponents;
+	TArray<FFastGeoSkinnedMeshComponent> SkinnedMeshComponents;
+	TArray<FFastGeoInstancedSkinnedMeshComponent> InstancedSkinnedMeshComponents;
+
+	template <typename ThisType, typename TComponent, typename TFunc>
+	static void ForEachComponent(ThisType* Self, TFunc&& InFunc)
+	{
+		auto ForEachArray = [&InFunc](auto& InArray)
+		{
+			typedef typename TDecay<decltype(InArray)>::Type ArrayType;
+			using TArrayComponent = typename ArrayType::ElementType;
+
+			if constexpr (std::is_same<TComponent, FFastGeoComponent>::value)
+			{
+				for (auto& Component : InArray)
+				{
+					InFunc(Component);
+				}
+			}
+			else if (TArrayComponent::Type.IsA(TComponent::Type))
+			{
+				for (auto& Component : InArray)
+				{
+					checkSlow(Component.template IsA<TComponent>());
+					Forward<TFunc>(InFunc)(*(TComponent*)(&Component));
+				}
+			}
+		};
+		
+		ForEachArray(Self->StaticMeshComponents);
+		ForEachArray(Self->InstancedStaticMeshComponents);
+		ForEachArray(Self->SkinnedMeshComponents);
+		ForEachArray(Self->InstancedSkinnedMeshComponents);
+	}
+
+	template <typename ThisType, typename TComponent, typename TFunc>
+	static bool ForEachComponentBreakable(ThisType* Self, TFunc&& InFunc)
+	{
+		auto ForEachArray = [&InFunc](auto& InArray) -> bool
+		{
+			typedef typename TDecay<decltype(InArray)>::Type ArrayType;
+			using TArrayComponent = typename ArrayType::ElementType;
+
+			if constexpr (std::is_same<TComponent, FFastGeoComponent>::value)
+			{
+				for (auto& Component : InArray)
+				{
+					if (!InFunc(Component))
+					{
+						return false;
+					}
+				}
+			}
+			else if (TArrayComponent::Type.IsA(TComponent::Type))
+			{
+				for (auto& Component : InArray)
+				{
+					checkSlow(Component.template IsA<TComponent>());
+					if (!Forward<TFunc>(InFunc)(*(TComponent*)(&Component)))
+					{
+						return false;
+					}
+				}
+			}
+
+			return true;
+		};
+		
+		return ForEachArray(Self->StaticMeshComponents) &&
+			   ForEachArray(Self->InstancedStaticMeshComponents) &&
+			   ForEachArray(Self->SkinnedMeshComponents) &&
+			   ForEachArray(Self->InstancedSkinnedMeshComponents);
+	}
+
+private:
+	void ForceUpdateVisibility(const TArray<FFastGeoPrimitiveComponent*>& Components, int32 UpdateCounter);
+	void UpdateVisibility_Internal(TArray<FFastGeoPrimitiveComponent*>&& ShowComponents, TArray<FFastGeoPrimitiveComponent*>&& HideComponents, int32 UpdateCounter = 0);
+	void InitializeDynamicProperties();
+	void SetOwnerContainer(UFastGeoContainer* InOwner);
+	void SetComponentClusterIndex(int32 InComponentClusterIndex);
+	friend class UFastGeoContainer;
+
+	friend FArchive& operator<<(FArchive& Ar, FFastGeoComponentCluster& ComponentCluster);
+};
\ No newline at end of file
diff --git a/Engine/Plugins/Experimental/FastGeoStreaming/Source/FastGeoStreaming/Internal/FastGeoContainer.h b/Engine/Plugins/Experimental/FastGeoStreaming/Source/FastGeoStreaming/Internal/FastGeoContainer.h
new file mode 100644
--- /dev/null
+++ b/Engine/Plugins/Experimental/FastGeoStreaming/Source/FastGeoStreaming/Internal/FastGeoContainer.h
@@ -0,0 +1,301 @@
+// Copyright Epic Games, Inc. All Rights Reserved.
+
+#pragma once
+
+#include "Engine/AssetUserData.h"
+#include "Physics/Experimental/AsyncPhysicsStateProcessorInterface.h"
+#include "PhysicsEngine/PhysicsBodyInstanceOwnerInterface.h"
+#include "PhysicsEngine/BodyInstance.h"
+#include "AI/Navigation/NavigationElement.h"
+#include "FastGeoComponentCluster.h"
+#include "FastGeoHLOD.h"
+#include "UObject/ObjectPtr.h"
+
+#if WITH_EDITOR
+#include "UObject/ObjectSaveContext.h"
+#endif // WITH_EDITOR
+
+#include "FastGeoContainer.generated.h"
+
+class UWorld;
+class ULevel;
+class URuntimeVirtualTexture;
+class UStaticMesh;
+class UMaterialInterface;
+class FFastGeoComponent;
+class FChaosUserDefinedEntity;
+
+struct FFastGeoComponentClusterPendingState
+{
+	struct FRenderState
+	{
+		FRenderState()
+		{
+			Reset();
+		}
+
+		void Reset()
+		{
+			ComponentsToProcess.Reset();
+			NumToProcess = 0;
+			NumProcessed = 0;
+			TotalNumProcessed = 0;
+		}
+
+		bool IsCompleted() const
+		{
+			return TotalNumProcessed >= ComponentsToProcess.Num();
+		}
+
+		TArray<FFastGeoPrimitiveComponent*> ComponentsToProcess;
+
+		int32 NumToProcess;
+		int32 NumProcessed;
+		int32 TotalNumProcessed;
+	} RenderState;
+
+	struct FPhysicsState
+	{
+		FPhysicsState()
+		{
+			Reset();
+		}
+
+		void Reset()
+		{
+			ComponentsToProcess = nullptr;
+			TotalNumProcessed = 0;
+		}
+
+		bool IsCompleted() const
+		{
+			return !ComponentsToProcess || (TotalNumProcessed >= ComponentsToProcess->Num());
+		}
+
+		TArray<FFastGeoComponent*>* ComponentsToProcess;
+		std::atomic<int32> TotalNumProcessed;
+	} PhysicsState;
+
+	bool bIsInBlockingWait;
+	
+	FFastGeoComponentClusterPendingState()
+	{
+		Reset();
+	}
+
+	void Reset()
+	{
+		RenderState.Reset();
+		PhysicsState.Reset();
+		bIsInBlockingWait = false;
+	}
+
+	bool HasAnyPendingState() const
+	{ 
+		return !RenderState.IsCompleted() || !PhysicsState.IsCompleted();
+	}
+};
+
+UCLASS(Within = Level)
+class FASTGEOSTREAMING_API UFastGeoContainer : public UAssetUserData, public IAsyncPhysicsStateProcessor, public IPhysicsBodyInstanceOwnerResolver
+{
+	GENERATED_BODY()
+
+public:
+	void Register();
+	void Unregister();
+	bool IsRegistered() const;
+	void Tick(bool bWaitForCompletion = false);
+	bool HasAnyPendingTasks() const;
+	bool HasAnyPendingCreateTasks() const;
+	bool HasAnyPendingDestroyTasks() const;
+	void PrecachePSOs();
+	UWorld* GetWorld() const;
+	ULevel* GetLevel() const;
+	FFastGeoComponentCluster* GetComponentCluster(uint32 InComponentClusterTypeID, int32 InComponentClusterIndex);
+
+	//~ Begin UObject Interface
+	virtual void BeginDestroy() override;
+	virtual bool IsReadyForFinishDestroy() override;
+	virtual void Serialize(FArchive& Ar) override;
+#if WITH_EDITOR
+	virtual void PreSave(FObjectPreSaveContext ObjectSaveContext) override;
+#endif
+	//~ End UObject Interface
+
+	//~Begin IPhysicsBodyInstanceOwnerResolver
+	virtual IPhysicsBodyInstanceOwner* ResolvePhysicsBodyInstanceOwner(Chaos::FConstPhysicsObjectHandle PhysicsObject) override;
+	//~End IPhysicsBodyInstanceOwnerResolver
+
+#if WITH_EDITOR
+	void OnCreated();
+	void AddComponentCluster(FFastGeoComponentCluster* ComponentCluster);
+#endif
+
+	void InitializeDynamicProperties(bool bInitForPlay = false);
+
+	void OnCreateRenderStateBegin_GameThread();
+	void OnDestroyRenderStateBegin_GameThread();
+	void OnCreateRenderStateEnd_GameThread();
+	void OnDestroyRenderStateEnd_GameThread();
+	void OnCreateRenderState_Concurrent();
+	void OnDestroyRenderState_Concurrent();
+
+	void OnCreatePhysicsStateBegin_GameThread();
+	void OnDestroyPhysicsStateBegin_GameThread();
+
+	//~ Begin IAsyncPhysicsStateProcessor
+	virtual bool AllowsAsyncPhysicsStateCreation() const override final;
+	virtual bool AllowsAsyncPhysicsStateDestruction() const override final;
+	virtual bool IsAsyncPhysicsStateCreated() const override final;
+	virtual UObject* GetAsyncPhysicsStateObject() const override final;
+	virtual void OnAsyncCreatePhysicsStateBegin_GameThread() override;
+	virtual bool OnAsyncCreatePhysicsState(const UE::FTimeout& Timeout) override;
+	virtual void OnAsyncCreatePhysicsStateEnd_GameThread() override;
+	virtual void OnAsyncDestroyPhysicsStateBegin_GameThread() override;
+	virtual bool OnAsyncDestroyPhysicsState(const UE::FTimeout& Timeout) override;
+	virtual void OnAsyncDestroyPhysicsStateEnd_GameThread() override;
+	virtual void CollectBodySetupsWithPhysicsMeshesToCreate(TSet<UBodySetup*>& OutBodySetups) const override;
+	//~ End IAsyncPhysicsStateProcessor
+
+	template <typename TComponentCluster = const FFastGeoComponentCluster, typename TFunc>
+	void ForEachComponentCluster(TFunc&& InFunc) const
+	{
+		ForEachComponentCluster<const UFastGeoContainer, const TComponentCluster, TFunc>(this, Forward<TFunc>(InFunc));
+	}
+
+	template <typename TComponentCluster = FFastGeoComponentCluster, typename TFunc>
+	void ForEachComponentCluster(TFunc&& InFunc)
+	{
+		ForEachComponentCluster<UFastGeoContainer, TComponentCluster, TFunc>(this, Forward<TFunc>(InFunc));
+	}
+
+	template <typename TComponentCluster = const FFastGeoComponentCluster, typename TFunc>
+	bool ForEachComponentClusterBreakable(TFunc&& InFunc) const
+	{
+		return ForEachComponentClusterBreakable<const UFastGeoContainer, const TComponentCluster, TFunc>(this, Forward<TFunc>(InFunc));
+	}
+
+	template <typename TComponentCluster = FFastGeoComponentCluster, typename TFunc>
+	bool ForEachComponentClusterBreakable(TFunc&& InFunc)
+	{
+		return ForEachComponentClusterBreakable<UFastGeoContainer, TComponentCluster, TFunc>(this, Forward<TFunc>(InFunc));
+	}
+
+protected:
+
+	FFastGeoComponentClusterPendingState PendingCreate;
+	FFastGeoComponentClusterPendingState PendingDestroy;
+
+private:
+
+#if WITH_EDITOR
+	void CreateEditorProxyObjects();
+	void CollectAssetReferences();
+#endif
+	void SerializeComponentClusters(FArchive& Ar);
+	void RegisterToNavigationSystem();
+	void UnregisterFromNavigationSystem();
+
+	template <typename ThisType, typename TComponentCluster, typename TFunc>
+	static void ForEachComponentCluster(ThisType* Self, TFunc&& InFunc)
+	{
+		auto ForEachArray = [&InFunc](auto& InArray)
+		{
+			typedef typename TDecay<decltype(InArray)>::Type ArrayType;
+			using TArrayComponentCluster = typename ArrayType::ElementType;
+
+			if constexpr (std::is_same<TComponentCluster, FFastGeoComponentCluster>::value)
+			{
+				for (auto& ComponentCluster : InArray)
+				{
+					InFunc(ComponentCluster);
+				}
+			}
+			else if (TArrayComponentCluster::Type.IsA(TComponentCluster::Type))
+			{
+				for (auto& ComponentCluster : InArray)
+				{
+					checkSlow(ComponentCluster.template IsA<TComponentCluster>());
+					Forward<TFunc>(InFunc)(*StaticCast<TComponentCluster*>(&ComponentCluster));
+				}
+			}
+		};
+
+		ForEachArray(Self->ComponentClusters);
+		ForEachArray(Self->HLODs);
+	}
+
+	template <typename ThisType, typename TComponentCluster, typename TFunc>
+	static bool ForEachComponentClusterBreakable(ThisType* Self, TFunc&& InFunc)
+	{
+		auto ForEachArray = [&InFunc](auto& InArray) -> bool
+		{
+			typedef typename TDecay<decltype(InArray)>::Type ArrayType;
+			using TArrayComponentCluster = typename ArrayType::ElementType;
+
+			if constexpr (std::is_same<TComponentCluster, FFastGeoComponentCluster>::value)
+			{
+				for (auto& ComponentCluster : InArray)
+				{
+					if (!InFunc(ComponentCluster))
+					{
+						return false;
+					}
+				}
+			}
+			else if (TArrayComponentCluster::Type.IsA(TComponentCluster::Type))
+			{
+				for (auto& ComponentCluster : InArray)
+				{
+					checkSlow(ComponentCluster.template IsA<TComponentCluster>());
+					if (!Forward<TFunc>(InFunc)(*StaticCast<TComponentCluster*>(&ComponentCluster)))
+					{
+						return false;
+					}
+				}
+			}
+
+			return true;
+		};
+
+		return ForEachArray(Self->ComponentClusters) &&
+			ForEachArray(Self->HLODs);
+	}
+
+	enum EPhysicsStateCreation
+	{
+		NotCreated,
+		Creating,
+		Created,
+		Destroying
+	};
+	EPhysicsStateCreation PhysicsStateCreation = EPhysicsStateCreation::NotCreated;
+
+	// Persistent data
+	TArray<FFastGeoComponentCluster> ComponentClusters;
+	TArray<FFastGeoHLOD> HLODs;
+
+	UPROPERTY()
+	TArray<TObjectPtr<UObject>> Assets;
+
+#if WITH_EDITORONLY_DATA
+	UPROPERTY(Transient)
+	TSet<TObjectPtr<UObject>> EditorProxyObjects;
+#endif
+
+#if !WITH_EDITOR && UE_WITH_PSO_PRECACHING
+	bool bPrecachedPSOs = false;
+	UE::Tasks::TTask<void> PrecachePSOsTask;
+#endif
+
+	// Transient data
+	TArray<FFastGeoComponent*> CollisionComponents;
+	TArray<FNavigationElementHandle> NavigationElementHandles;
+
+	bool bIsRegistered = false;
+
+	FRenderCommandFence DestroyFence;
+
+	friend class FFastGeoGatherFastGeoContainerAssetRefsArchive;
+};
\ No newline at end of file
diff --git a/Engine/Plugins/Experimental/FastGeoStreaming/Source/FastGeoStreaming/Internal/FastGeoElementType.h b/Engine/Plugins/Experimental/FastGeoStreaming/Source/FastGeoStreaming/Internal/FastGeoElementType.h
new file mode 100644
--- /dev/null
+++ b/Engine/Plugins/Experimental/FastGeoStreaming/Source/FastGeoStreaming/Internal/FastGeoElementType.h
@@ -0,0 +1,44 @@
+// Copyright Epic Games, Inc. All Rights Reserved.
+
+#pragma once
+
+#include "HAL/Platform.h"
+
+struct FFastGeoElementType
+{
+public:
+	static const FFastGeoElementType Invalid;
+
+	explicit FFastGeoElementType(const FFastGeoElementType* Parent = nullptr) : ID(++NextUniqueID), ParentType(Parent) {}
+
+	bool operator==(const FFastGeoElementType& Other) const
+	{
+		return IsSameTypeID(Other.ID);
+	}
+
+	bool IsA(const FFastGeoElementType& Other) const
+	{
+		return IsSameTypeID(Other.ID) || (ParentType && ParentType->IsA(Other));
+	}
+
+	/** Returns true if element has the same type ID than the specified type */
+	bool IsSameTypeID(uint32 InID) const
+	{
+		return ID == InID;
+	}
+
+	bool IsValid() const
+	{
+		return !IsSameTypeID(Invalid.ID);
+	}
+
+private:
+	enum class EInitMode { Invalid };
+	explicit FFastGeoElementType(EInitMode) : ID(0xFFFFFFFF), ParentType(nullptr) {}
+
+	static uint32 NextUniqueID;
+	uint32 ID;
+	const FFastGeoElementType* ParentType;
+
+	friend class IFastGeoElement;
+};
\ No newline at end of file
diff --git a/Engine/Plugins/Experimental/FastGeoStreaming/Source/FastGeoStreaming/Internal/FastGeoHLOD.h b/Engine/Plugins/Experimental/FastGeoStreaming/Source/FastGeoStreaming/Internal/FastGeoHLOD.h
new file mode 100644
--- /dev/null
+++ b/Engine/Plugins/Experimental/FastGeoStreaming/Source/FastGeoStreaming/Internal/FastGeoHLOD.h
@@ -0,0 +1,59 @@
+// Copyright Epic Games, Inc. All Rights Reserved.
+
+#pragma once
+
+#include "FastGeoComponentCluster.h"
+#include "WorldPartition/HLOD/IWorldPartitionHLODObject.h"
+
+class FFastGeoHLOD : public FFastGeoComponentCluster, public IWorldPartitionHLODObject
+{
+public:
+	typedef FFastGeoComponentCluster Super;
+
+	/** Static type identifier for this element class */
+	static const FFastGeoElementType Type;
+
+	FFastGeoHLOD(UFastGeoContainer* InOwner = nullptr, FName InName = NAME_None, FFastGeoElementType InType = Type);
+	~FFastGeoHLOD() {}
+
+	virtual bool IsVisible() const override { return bIsVisible; }
+
+	// ~Begin FFastGeoComponentCluster interface
+	virtual void OnRegister() override;
+	virtual void OnUnregister() override;
+	virtual void Serialize(FArchive& Ar) override;
+	// ~End FFastGeoComponentCluster interface
+
+	// ~Begin IWorldPartitionHLODObject interface
+	virtual UObject* GetUObject() const override { return nullptr; }
+	virtual ULevel* GetHLODLevel() const override;
+	virtual FString GetHLODNameOrLabel() const override;
+	virtual bool DoesRequireWarmup() const override;
+	virtual TSet<UObject*> GetAssetsToWarmup() const override;
+	virtual void SetVisibility(bool bIsVisible) override;
+	virtual const FGuid& GetSourceCellGuid() const override;
+	virtual bool IsStandalone() const override;
+	virtual const FGuid& GetStandaloneHLODGuid() const override;
+	// ~End IWorldPartitionHLODObject interface
+
+#if WITH_EDITOR
+	// ~Begin FFastGeoComponentCluster interface
+	virtual FFastGeoComponent& AddComponent(FFastGeoElementType InComponentType) override;
+	virtual void PreSave(FObjectPreSaveContext ObjectSaveContext) override;
+	// ~End FFastGeoComponentCluster interface
+
+	void SetSourceCellGuid(const FGuid& InSourceCellGuid);
+	void SetRequireWarmup(bool bInRequireWarmup);
+	void SetStandaloneHLODGuid(const FGuid& InStandaloneHLODGuid);
+#endif
+
+private:
+
+	// Transient Data
+	bool bIsVisible;
+
+	// Persistent Data
+	bool bRequireWarmup;
+	FGuid SourceCellGuid;
+	FGuid StandaloneHLODGuid;
+};
\ No newline at end of file
diff --git a/Engine/Plugins/Experimental/FastGeoStreaming/Source/FastGeoStreaming/Internal/FastGeoInstancedSkinnedMeshComponent.h b/Engine/Plugins/Experimental/FastGeoStreaming/Source/FastGeoStreaming/Internal/FastGeoInstancedSkinnedMeshComponent.h
new file mode 100644
--- /dev/null
+++ b/Engine/Plugins/Experimental/FastGeoStreaming/Source/FastGeoStreaming/Internal/FastGeoInstancedSkinnedMeshComponent.h
@@ -0,0 +1,68 @@
+// Copyright Epic Games, Inc. All Rights Reserved.
+
+#pragma once
+
+#include "FastGeoSkinnedMeshComponent.h"
+#include "Components/InstancedSkinnedMeshComponent.h"
+#include "InstanceData/InstanceDataManager.h"
+#include "InstancedSkinnedMeshSceneProxyDesc.h"
+
+class FASTGEOSTREAMING_API FFastGeoInstancedSkinnedMeshComponent : public FFastGeoSkinnedMeshComponentBase
+{
+public:
+	typedef FFastGeoSkinnedMeshComponentBase Super;
+
+	/** Static type identifier for this element class */
+	static const FFastGeoElementType Type;
+
+	FFastGeoInstancedSkinnedMeshComponent(int32 InComponentIndex = INDEX_NONE, FFastGeoElementType InType = Type);
+	virtual ~FFastGeoInstancedSkinnedMeshComponent() = default;
+
+protected:
+	//~ Being FFastGeoComponent interface
+	virtual void Serialize(FArchive& Ar) override;
+#if WITH_EDITOR
+	virtual void InitializeFromComponent(UActorComponent* Component) override;
+#endif
+	//~ End FFastGeoComponent interface
+
+	//~ Begin FFastGeoPrimitiveComponent interface
+	virtual FPrimitiveSceneProxyDesc& GetSceneProxyDesc() override { return SceneProxyDesc; }
+	virtual const FPrimitiveSceneProxyDesc& GetSceneProxyDesc() const override { return SceneProxyDesc; }
+#if WITH_EDITOR
+	virtual void InitializeSceneProxyDescFromComponent(UActorComponent* Component) override;
+	virtual void ResetSceneProxyDescUnsupportedProperties() override;
+#endif
+	virtual void InitializeSceneProxyDescDynamicProperties() override;
+	virtual void ApplyWorldTransform(const FTransform& InTransform) override;
+	//~ End FFastGeoPrimitiveComponent interface
+
+	//~ Begin FFastGeoSkinnedMeshComponentBase interface
+	virtual FSkinnedMeshSceneProxyDesc& GetSkinnedMeshSceneProxyDesc() override { return SceneProxyDesc; }
+	virtual const FSkinnedMeshSceneProxyDesc& GetSkinnedMeshSceneProxyDesc() const override { return SceneProxyDesc; }
+	virtual void UpdateSkinning() override {}
+	virtual FSkeletalMeshObject* CreateMeshObject() override;
+	virtual FPrimitiveSceneProxy* AllocateSceneProxy() override;
+
+	int32 GetInstanceCount() const
+	{
+		return InstanceData.Num();
+	}
+
+	UTransformProviderData* GetTransformProvider() const
+	{
+		return nullptr;
+	}
+	//~ End FFastGeoSkinnedMeshComponentBase interface
+
+private:
+	TArray<FSkinnedMeshInstanceData> InstanceData;
+	int32 NumCustomDataFloats = 0;
+	TArray<float> InstanceCustomData;
+
+	FInstancedSkinnedMeshSceneProxyDesc SceneProxyDesc;
+
+	TSharedPtr<FInstanceDataManager> InstanceDataManager;
+
+	friend class FInstancedSkinnedMeshComponentHelper;
+};
diff --git a/Engine/Plugins/Experimental/FastGeoStreaming/Source/FastGeoStreaming/Internal/FastGeoInstancedStaticMeshComponent.h b/Engine/Plugins/Experimental/FastGeoStreaming/Source/FastGeoStreaming/Internal/FastGeoInstancedStaticMeshComponent.h
new file mode 100644
--- /dev/null
+++ b/Engine/Plugins/Experimental/FastGeoStreaming/Source/FastGeoStreaming/Internal/FastGeoInstancedStaticMeshComponent.h
@@ -0,0 +1,107 @@
+// Copyright Epic Games, Inc. All Rights Reserved.
+
+#pragma once
+
+#include "FastGeoStaticMeshComponent.h"
+#include "InstancedStaticMeshSceneProxyDesc.h"
+#include "FastGeoInstancedStaticMeshComponent.generated.h"
+
+class FASTGEOSTREAMING_API FFastGeoInstancedStaticMeshComponent : public FFastGeoStaticMeshComponentBase
+{
+public:
+	typedef FFastGeoStaticMeshComponentBase Super;
+
+	/** Static type identifier for this element class */
+	static const FFastGeoElementType Type;
+
+	FFastGeoInstancedStaticMeshComponent(int32 InComponentIndex = INDEX_NONE, FFastGeoElementType InType = Type);
+	virtual ~FFastGeoInstancedStaticMeshComponent() = default;
+
+	//~ Begin FFastGeoComponent Interface
+	virtual void Serialize(FArchive& Ar) override;
+	virtual void OnAsyncCreatePhysicsState() override;
+	virtual void OnAsyncDestroyPhysicsState() override;
+	virtual void OnAsyncDestroyPhysicsStateBegin_GameThread() override;
+	virtual void OnAsyncDestroyPhysicsStateEnd_GameThread() override;
+#if WITH_EDITOR
+	virtual void InitializeFromComponent(UActorComponent* Component) override;
+	virtual UClass* GetEditorProxyClass() const override;
+#endif
+	//~ End FFastGeoComponent Interface
+
+	//~ Begin FFastGeoPrimitiveComponent interface
+	virtual void GetNavigationData(FNavigationRelevantData& Data) const override;
+	virtual FBox GetNavigationBounds() const override;
+	virtual bool IsNavigationRelevant() const override;
+	virtual bool DoCustomNavigableGeometryExport(FNavigableGeometryExport& GeomExport) const override;
+	virtual bool ShouldSkipNavigationDirtyAreaOnAddOrRemove() override { return true; }
+	//~ End FFastGeoPrimitiveComponent interface
+
+protected:
+	//~ Begin FFastGeoPrimitiveComponent interface
+	virtual FPrimitiveSceneProxyDesc& GetSceneProxyDesc() override { return SceneProxyDesc; }
+	virtual const FPrimitiveSceneProxyDesc& GetSceneProxyDesc() const override { return SceneProxyDesc; }
+#if WITH_EDITOR
+	virtual void InitializeSceneProxyDescFromComponent(UActorComponent* Component) override;
+	virtual void ResetSceneProxyDescUnsupportedProperties() override;
+#endif
+	virtual void ApplyWorldTransform(const FTransform& InTransform) override;
+	//~ End FFastGeoPrimitiveComponent interface
+
+	//~ Begin FFastGeoStaticMeshComponentBase interface
+	virtual FStaticMeshSceneProxyDesc& GetStaticMeshSceneProxyDesc() override { return SceneProxyDesc; }
+	virtual const FStaticMeshSceneProxyDesc& GetStaticMeshSceneProxyDesc() const override { return SceneProxyDesc; }
+	virtual FPrimitiveSceneProxy* CreateStaticMeshSceneProxy(const Nanite::FMaterialAudit& NaniteMaterials, bool bCreateNanite) override;
+#if !WITH_EDITOR && UE_WITH_PSO_PRECACHING
+	virtual void CollectPSOPrecacheData(const FPSOPrecacheParams& BasePrecachePSOParams, FMaterialInterfacePSOPrecacheParamsList& OutParams) override;
+#endif
+	//~ End FFastGeoStaticMeshComponentBase interface
+
+	TSharedPtr<FInstanceDataSceneProxy, ESPMode::ThreadSafe>& BuildInstanceData();
+
+private:
+	enum class EBoundsType
+	{
+		LocalBounds,
+		WorldBounds,
+		NavigationBounds
+	};
+	FBoxSphereBounds CalculateBounds(EBoundsType BoundsType);
+	void CreateAllInstanceBodies();
+	static TArray<Chaos::FPhysicsObject*> GetAllPhysicsObjects(TArray<FBodyInstance*> InInstanceBodies);
+
+private:
+	// Persistent Data
+	TArray<FInstancedStaticMeshInstanceData> PerInstanceSMData;
+	int32 InstancingRandomSeed = 0;
+	TArray<float> PerInstanceSMCustomData;
+	TArray<FInstancedStaticMeshRandomSeed> AdditionalRandomSeeds;
+	FBox NavigationBounds;
+	float InstanceLODDistanceScale = 1.f;
+	int32 InstanceStartCullDistance = 0;
+	int32 InstanceEndCullDistance = 0;
+	uint8 bUseGpuLodSelection : 1 = false;
+	FInstancedStaticMeshSceneProxyDesc SceneProxyDesc{};
+
+	// Transient data
+	TSharedPtr<FInstanceDataSceneProxy, ESPMode::ThreadSafe> DataProxy{};
+	TArray<float> InstanceRandomIDs;
+	
+	/** Physics representation of the instance bodies. */
+	TArray<FBodyInstance*> InstanceBodies;
+
+	/** Payload used by asynchronous destruction of physics state (see OnAsyncDestroyPhysicsState). */
+	TArray<FBodyInstance*> AsyncDestroyPhysicsStatePayload;
+
+	friend class FInstancedStaticMeshComponentHelper;
+};
+
+UCLASS()
+class UFastGeoInstancedStaticMeshComponentEditorProxy : public UFastGeoStaticMeshComponentEditorProxy
+{
+	GENERATED_BODY()
+
+#if WITH_EDITOR
+	typedef FFastGeoInstancedStaticMeshComponent ComponentType;
+#endif
+};
diff --git a/Engine/Plugins/Experimental/FastGeoStreaming/Source/FastGeoStreaming/Internal/FastGeoMeshComponent.h b/Engine/Plugins/Experimental/FastGeoStreaming/Source/FastGeoStreaming/Internal/FastGeoMeshComponent.h
new file mode 100644
--- /dev/null
+++ b/Engine/Plugins/Experimental/FastGeoStreaming/Source/FastGeoStreaming/Internal/FastGeoMeshComponent.h
@@ -0,0 +1,43 @@
+// Copyright Epic Games, Inc. All Rights Reserved.
+
+#pragma once
+
+#include "FastGeoPrimitiveComponent.h"
+#include "Components/MeshComponent.h"
+#include "MeshComponentHelper.h"
+
+class FASTGEOSTREAMING_API FFastGeoMeshComponent : public FFastGeoPrimitiveComponent
+{
+public:
+	typedef FFastGeoPrimitiveComponent Super;
+	
+	/** Static type identifier for this element class */
+	static const FFastGeoElementType Type;
+
+	FFastGeoMeshComponent(int32 InComponentIndex = INDEX_NONE, FFastGeoElementType InType = Type);
+	virtual ~FFastGeoMeshComponent() = default;
+
+	//~ Begin FFastGeoComponent interface
+	virtual void Serialize(FArchive& Ar) override;
+#if WITH_EDITOR
+	virtual void InitializeFromComponent(UActorComponent* InComponent) override;
+#endif
+	//~ End FFastGeoComponent interface
+
+	FMaterialRelevance GetMaterialRelevance(EShaderPlatform InShaderPlatform) const;
+
+	void GetMaterialSlotsOverlayMaterial(TArray<TObjectPtr<UMaterialInterface>>& OutMaterialSlotOverlayMaterials) const;
+
+	//Get the asset material slot overlay material
+	virtual void GetDefaultMaterialSlotsOverlayMaterial(TArray<TObjectPtr<UMaterialInterface>>& AssetMaterialSlotOverlayMaterials) const = 0;
+
+	virtual const TArray<TObjectPtr<UMaterialInterface>>& GetComponentMaterialSlotsOverlayMaterial() const = 0;
+	virtual UMaterialInterface* GetOverlayMaterial() const = 0;
+
+protected:
+	// Persistent Data
+	TArray<TObjectPtr<UMaterialInterface>> OverrideMaterials;
+
+private:
+	friend class FMeshComponentHelper;
+};
\ No newline at end of file
diff --git a/Engine/Plugins/Experimental/FastGeoStreaming/Source/FastGeoStreaming/Internal/FastGeoPrimitiveComponent.h b/Engine/Plugins/Experimental/FastGeoStreaming/Source/FastGeoStreaming/Internal/FastGeoPrimitiveComponent.h
new file mode 100644
--- /dev/null
+++ b/Engine/Plugins/Experimental/FastGeoStreaming/Source/FastGeoStreaming/Internal/FastGeoPrimitiveComponent.h
@@ -0,0 +1,304 @@
+// Copyright Epic Games, Inc. All Rights Reserved.
+
+#pragma once
+
+#include "FastGeoComponent.h"
+#include "Components/PrimitiveComponent.h"
+#include "Components/ComponentInterfaces.h"
+#include "Engine/EngineTypes.h"
+#include "PrimitiveSceneInfoData.h"
+#include "PrimitiveSceneProxy.h"
+#include "RenderCommandFence.h"
+#include "PhysicsEngine/BodyInstance.h"
+#include "PhysicsEngine/PhysicsBodyInstanceOwnerInterface.h"
+#include "Chaos/ChaosUserEntity.h"
+#include "PSOPrecacheFwd.h"
+#include "Templates/DontCopy.h"
+#include "FastGeoPrimitiveComponent.generated.h"
+
+class FPrimitiveSceneProxy;
+class UPrimitiveComponent;
+class URuntimeVirtualTexture;
+class FFastGeoPrimitiveComponent;
+struct FPrimitiveSceneDesc;
+struct FPrimitiveSceneProxyDesc;
+struct FPSOPrecacheParams;
+enum class EPSOPrecachePriority : uint8;
+
+enum class ESceneProxyCreationError
+{
+	None,
+	WaitingPSOs,
+	InvalidMesh
+};
+
+class FFastGeoPhysicsBodyInstanceOwner : public FChaosUserDefinedEntity, public IPhysicsBodyInstanceOwner
+{
+public:
+	FFastGeoPhysicsBodyInstanceOwner();
+	virtual ~FFastGeoPhysicsBodyInstanceOwner() = default;
+
+	/** Returns the IPhysicsBodyInstanceOwner based on the provided FChaosUserDefinedEntity */
+	static IPhysicsBodyInstanceOwner* GetPhysicsBodyInstanceOwner(FChaosUserDefinedEntity* InUserDefinedEntity);
+
+	//~ Begin FChaosUserDefinedEntity interface
+	TWeakObjectPtr<UObject> GetOwnerObject() override;
+	//~ End FChaosUserDefinedEntity interface
+
+	//~ Begin IPhysicsBodyInstanceOwner interface
+	virtual bool IsStaticPhysics() const override;
+	virtual UObject* GetSourceObject() const override;
+	virtual ECollisionResponse GetCollisionResponseToChannel(ECollisionChannel Channel) const override;
+	virtual UPhysicalMaterial* GetPhysicalMaterial() const override;
+	virtual void GetComplexPhysicalMaterials(TArray<UPhysicalMaterial*>& OutPhysMaterials, TArray<FPhysicalMaterialMaskParams>* OutPhysMaterialMasks) const override;
+	//~ End IPhysicsBodyInstanceOwner interface
+
+private:
+	void Uninitialize();
+	void Initialize(FFastGeoPrimitiveComponent* InOwner);
+
+	FFastGeoPrimitiveComponent* OwnerComponent;
+	TWeakObjectPtr<UFastGeoContainer> OwnerContainer;
+	static const FName NAME_FastGeoPhysicsBodyInstanceOwner;
+
+	friend class FFastGeoPrimitiveComponent;
+	friend class FFastGeoInstancedStaticMeshComponent;
+};
+
+class FASTGEOSTREAMING_API FFastGeoPrimitiveComponent : public FFastGeoComponent
+{
+public:
+	typedef FFastGeoComponent Super;
+	typedef UFastGeoPrimitiveComponentEditorProxy EditorProxyType;
+
+	/** Static type identifier for this element class */
+	static const FFastGeoElementType Type;
+
+	FFastGeoPrimitiveComponent(int32 InComponentIndex = INDEX_NONE, FFastGeoElementType InType = Type);
+	virtual ~FFastGeoPrimitiveComponent() = default;
+
+	TArray<URuntimeVirtualTexture*> const& GetRuntimeVirtualTextures() const { return RuntimeVirtualTextures; }
+	bool IsFirstPersonRelevant() const;
+
+	FPrimitiveComponentId GetPrimitiveSceneId() const { return PrimitiveSceneData.PrimitiveSceneId; }
+	virtual UObject const* AdditionalStatObject() const { return nullptr; }
+		
+	FSceneInterface* GetScene() const;
+	FPrimitiveSceneProxy* GetSceneProxy() const;
+	const FTransform& GetTransform() const;
+	const FBoxSphereBounds& GetBounds() const;
+	FMatrix GetRenderMatrix() const;
+	float GetLastRenderTimeOnScreen() const;
+
+	// Used by FStaticMeshComponentHelper/FInstancedStaticMeshComponentHelper
+	const FTransform& GetComponentTransform() const { return GetTransform(); }
+
+	void CreateRenderState(FRegisterComponentContext* Context);
+
+	struct FFastGeoDestroyRenderStateContext
+	{
+		FFastGeoDestroyRenderStateContext(FSceneInterface* InScene);
+		~FFastGeoDestroyRenderStateContext();
+		bool HasPendingWork() const;
+
+		static void DestroyProxy(FFastGeoDestroyRenderStateContext* InContext, FPrimitiveSceneProxy* InPrimitiveSceneProxy);
+
+	private:
+		FSceneInterface* Scene;
+		TArray<FPrimitiveSceneProxy*> PrimitiveSceneProxies;
+	};
+
+	virtual void DestroyRenderState(FFastGeoDestroyRenderStateContext* Context);
+	bool IsRenderStateCreated() const;
+	bool IsRenderStateDelayed() const;
+	bool IsRenderStateDirty() const;
+	bool IsDrawnInGame() const;
+	bool ShouldCreateRenderState() const;
+	void MarkRenderStateDirty();
+	void MarkPrecachePSOsRequired();
+	void PrecachePSOs();
+	bool CheckPSOPrecachingAndBoostPriority(EPSOPrecachePriority NewPSOPrecachePriority = EPSOPrecachePriority::High);
+	bool IsPSOPrecaching() const;
+	void SetCollisionEnabled(bool bInCollisionEnabled);
+	void UpdateVisibility();
+	EComponentMobility::Type GetMobility() const;
+	
+	//~ Begin FFastGeoComponent interface
+	virtual void Serialize(FArchive& Ar) override;
+	virtual void InitializeDynamicProperties() override;
+	virtual void OnAsyncCreatePhysicsState() override;
+	virtual void OnAsyncDestroyPhysicsStateBegin_GameThread() override;
+	virtual void OnAsyncDestroyPhysicsStateEnd_GameThread() override;
+	virtual void OnAsyncDestroyPhysicsState() override;
+	virtual bool IsCollisionEnabled() const override;
+#if WITH_EDITOR
+	virtual void InitializeFromComponent(UActorComponent* Component) override;
+	virtual UClass* GetEditorProxyClass() const override;
+#endif
+	//~End FFastGeoComponent interface
+
+	//~ Begin Materials
+	virtual UMaterialInterface* GetMaterial(int32 ElementIndex) const = 0;
+	virtual int32 GetNumMaterials() const = 0;
+	virtual void GetUsedMaterials(TArray<UMaterialInterface*>& OutMaterials, bool bGetDebugMaterials = false) const = 0;
+	UE_DEPRECATED(5.7, "Please use GetUsedMaterialPropertyDesc with EShaderPlatform argument and not ERHIFeatureLevel::Type")
+	FPrimitiveMaterialPropertyDescriptor GetUsedMaterialPropertyDesc(ERHIFeatureLevel::Type FeatureLevel) const;
+	FPrimitiveMaterialPropertyDescriptor GetUsedMaterialPropertyDesc(EShaderPlatform InShaderPlatform) const;
+	//~ End Materials
+
+	//~ Begin Navigation
+	virtual bool IsNavigationRelevant() const;
+	virtual bool ShouldSkipNavigationDirtyAreaOnAddOrRemove() { return false; }
+	virtual FBox GetNavigationBounds() const;
+	virtual void GetNavigationData(FNavigationRelevantData& OutData) const;
+	virtual EHasCustomNavigableGeometry::Type HasCustomNavigableGeometry() const;
+	virtual bool DoCustomNavigableGeometryExport(FNavigableGeometryExport& GeomExport) const { return true; }
+	//~ End Navigation
+
+	
+protected:
+	virtual FPrimitiveSceneProxyDesc& GetSceneProxyDesc() = 0;
+	virtual const FPrimitiveSceneProxyDesc& GetSceneProxyDesc() const = 0;
+#if WITH_EDITOR
+	virtual void InitializeSceneProxyDescFromComponent(UActorComponent* Component) = 0;
+	virtual void ResetSceneProxyDescUnsupportedProperties();
+#endif
+	virtual void InitializeSceneProxyDescDynamicProperties();
+	virtual void ApplyWorldTransform(const FTransform& InTransform);
+	virtual FPrimitiveSceneProxy* CreateSceneProxy(ESceneProxyCreationError* OutError = nullptr) = 0;
+	FPrimitiveSceneDesc BuildSceneDesc();
+
+	//~ Begin Physics
+	virtual bool IsStaticPhysics() const;
+	virtual UObject* GetSourceObject() const;
+	virtual UPhysicalMaterial* GetPhysicalMaterial() const { return nullptr; }
+	virtual void GetComplexPhysicalMaterials(TArray<UPhysicalMaterial*>& OutPhysMaterials, TArray<FPhysicalMaterialMaskParams>* OutPhysMaterialMasks) const {}
+	virtual ECollisionResponse GetCollisionResponseToChannel(ECollisionChannel Channel) const;
+	//~ End Navigation
+
+protected:
+	// Persistent data
+	FTransform LocalTransform;
+	FTransform WorldTransform;
+	FBoxSphereBounds LocalBounds;
+	FBoxSphereBounds WorldBounds;
+	uint8 bIsVisible : 1 = true;
+	uint8 bStaticWhenNotMoveable : 1 = true;
+	uint8 bFillCollisionUnderneathForNavmesh : 1 = false;
+	uint8 bRasterizeAsFilledConvexVolume : 1 = false;
+	uint8 bCanEverAffectNavigation : 1 = false;
+	FCustomPrimitiveData CustomPrimitiveData;
+	TEnumAsByte<enum EDetailMode> DetailMode;
+	TEnumAsByte<EHasCustomNavigableGeometry::Type> bHasCustomNavigableGeometry;
+	FBodyInstance BodyInstance;
+	TArray<URuntimeVirtualTexture*> RuntimeVirtualTextures;
+
+	// Runtime Data (transient)
+	FFastGeoPhysicsBodyInstanceOwner BodyInstanceOwner;
+	FPrimitiveSceneInfoData PrimitiveSceneData{};
+	/** Payload used to release BodyInstance resources in asynchronous mode (see OnAsyncDestroyPhysicsState). */
+	TOptional<FBodyInstance::FAsyncTermBodyPayload> AsyncTermBodyPayload;
+
+	enum class EProxyCreationState : uint8
+	{
+		None, // Constructed/Initialized
+		Pending, // AddedToWorld & proxy creation is pending
+		Creating, // Actively creating the proxy
+		Created, // Proxy is now created
+		Delayed // Proxy creation delayed (used when PSO precaching is not ready when creating proxy)
+	};
+
+	EProxyCreationState ProxyState = EProxyCreationState::None;
+	bool bRenderStateDirty = false;
+
+private:
+	bool ShouldRenderProxyFallbackToDefaultMaterial() const;
+
+#if !WITH_EDITOR && UE_WITH_PSO_PRECACHING
+public:
+	void OnPrecacheFinished(int32 JobSetThatJustCompleted);
+
+protected:
+	void RequestRecreateRenderStateWhenPSOPrecacheFinished(const FGraphEventArray& PSOPrecacheCompileEvents);
+	void SetupPrecachePSOParams(FPSOPrecacheParams& Params);
+	virtual void CollectPSOPrecacheData(const FPSOPrecacheParams& BasePrecachePSOParams, FMaterialInterfacePSOPrecacheParamsList& OutParams) {}
+
+	// Cached array of material PSO requests which can be used to boost the priority
+	TArray<FMaterialPSOPrecacheRequestID> MaterialPSOPrecacheRequestIDs;
+	// Atomic int used to track the last PSO precache events
+	std::atomic<int> LatestPSOPrecacheJobSetCompleted = 0;
+	int32 LatestPSOPrecacheJobSet = 0;
+	// Helper flag to check if PSOs have been precached already
+	std::atomic<bool> bPSOPrecacheCalled = false;
+	std::atomic<bool> bPSOPrecacheRequired = false;
+	// PSOs requested priority
+	std::atomic<EPSOPrecachePriority> PSOPrecacheRequestPriority = EPSOPrecachePriority::Medium;
+	static_assert((int)EPSOPrecachePriority::Highest < 1 << 2);
+#endif
+
+	TDontCopy<FRWLock> Lock;
+
+	friend class FFastGeoComponentCluster;
+	friend class FPrimitiveComponentHelper;
+	friend class FFastGeoPhysicsBodyInstanceOwner;
+	friend class UFastGeoContainer;
+	friend class UFastGeoPrimitiveComponentEditorProxy;
+};
+
+// Dummy type to use as base class in non editor builds
+template <typename TInterface>
+class IFastGeoDummyInterface {};
+
+// Implement TInterface, only in editor.
+// Would have been easier with a #ifdef in the class declaration, but UHT doesn't allow it.
+template <typename TInterface>
+class IFastGeoEditorInterface : public std::conditional<WITH_EDITOR, TInterface, IFastGeoDummyInterface<TInterface>>::type {};
+
+UCLASS()
+class UFastGeoPrimitiveComponentEditorProxy : public UFastGeoComponentEditorProxy, public IFastGeoEditorInterface<IPrimitiveComponent>
+{
+	GENERATED_BODY()
+
+#if WITH_EDITOR
+	typedef FFastGeoPrimitiveComponent ComponentType;
+
+public:
+	virtual void NotifyRenderStateChanged();
+	virtual IPrimitiveComponent* GetPrimitiveComponentInterface();
+
+protected:
+	//~ Begin UObject interface
+	virtual void BeginDestroy() override;
+	//~ End UObject interface
+
+	//~ Begin IPrimitiveComponent interface
+	virtual bool IsRenderStateCreated() const override;
+	virtual bool IsRenderStateDirty() const override;
+	virtual bool ShouldCreateRenderState() const override;
+	virtual bool IsRegistered() const override;
+	virtual bool IsUnreachable() const override;
+	virtual UWorld* GetWorld() const override;
+	virtual FSceneInterface* GetScene() const override;
+	virtual FPrimitiveSceneProxy* GetSceneProxy() const override;
+	virtual void GetStreamableRenderAssetInfo(TArray<FStreamingRenderAssetPrimitiveInfo>& ) const override {} 
+	virtual void GetUsedMaterials(TArray<UMaterialInterface*>& OutMaterials, bool bGetDebugMaterials = false) const override;
+	virtual void MarkRenderStateDirty() override;
+	virtual void DestroyRenderState() override;
+	virtual void CreateRenderState(FRegisterComponentContext* Context) override;
+	virtual FString GetName() const override;
+	virtual FString GetFullName() const override;
+	virtual FTransform GetTransform() const override;
+	virtual FBoxSphereBounds GetBounds() const override;
+	virtual float GetLastRenderTimeOnScreen() const override;
+	virtual void GetPrimitiveStats(FPrimitiveStats& PrimitiveStats) const override;
+	virtual UObject* GetUObject() override;
+	virtual const UObject* GetUObject() const override;
+	virtual void PrecachePSOs() override;
+	virtual UObject* GetOwner() const override;
+	virtual FString GetOwnerName() const override;
+	virtual FPrimitiveSceneProxy* CreateSceneProxy() override;
+	virtual HHitProxy* CreateMeshHitProxy(int32 SectionIndex, int32 MaterialIndex) override;
+	virtual HHitProxy* CreatePrimitiveHitProxies(TArray<TRefCountPtr<HHitProxy> >& OutHitProxies) override;
+	//~ End IPrimitiveComponent interface
+#endif
+};
diff --git a/Engine/Plugins/Experimental/FastGeoStreaming/Source/FastGeoStreaming/Internal/FastGeoSkinnedMeshComponent.h b/Engine/Plugins/Experimental/FastGeoStreaming/Source/FastGeoStreaming/Internal/FastGeoSkinnedMeshComponent.h
new file mode 100644
--- /dev/null
+++ b/Engine/Plugins/Experimental/FastGeoStreaming/Source/FastGeoStreaming/Internal/FastGeoSkinnedMeshComponent.h
@@ -0,0 +1,120 @@
+// Copyright Epic Games, Inc. All Rights Reserved.
+
+#pragma once
+
+#include "AnimationRuntime.h"
+#include "Components/SkinnedMeshComponent.h"
+#include "FastGeoMeshComponent.h"
+#include "SkinnedMeshSceneProxyDesc.h"
+
+class UPhysicsAsset;
+class USkinnedAsset;
+class FSkeletalMeshObject;
+class FSkeletalMeshRenderData;
+
+class FASTGEOSTREAMING_API FFastGeoSkinnedMeshComponentBase : public FFastGeoMeshComponent
+{
+public:
+	typedef FFastGeoMeshComponent Super;
+
+	/** Static type identifier for this element class */
+	static const FFastGeoElementType Type;
+
+	FFastGeoSkinnedMeshComponentBase(int32 InComponentIndex = INDEX_NONE, FFastGeoElementType InType = Type);
+	virtual ~FFastGeoSkinnedMeshComponentBase() = default;
+
+	//~ Begin FFastGeoComponent interface
+	virtual void Serialize(FArchive& Ar) override;
+	virtual UBodySetup* GetBodySetup() const override;
+#if WITH_EDITOR
+	virtual void InitializeFromComponent(UActorComponent* Component) override;
+#endif
+	//~ End FFastGeoComponent interface
+	
+	//~ Begin FFastGeoPrimitiveComponent interface
+	virtual UMaterialInterface* GetMaterial(int32 MaterialIndex) const override;
+	virtual int32 GetNumMaterials() const override;
+	virtual void GetUsedMaterials(TArray<UMaterialInterface*>& OutMaterials, bool bGetDebugMaterials = false) const override;
+	//~ End FFastGeoPrimitiveComponent interface
+
+protected:
+	//~ Begin FFastGeoPrimitiveComponent interface
+#if WITH_EDITOR
+	virtual void InitializeSceneProxyDescFromComponent(UActorComponent* Component) override;
+	virtual void ResetSceneProxyDescUnsupportedProperties() override;
+#endif
+	virtual void InitializeSceneProxyDescDynamicProperties() override;
+	virtual FPrimitiveSceneProxy* CreateSceneProxy(ESceneProxyCreationError* OutError = nullptr) override;
+	virtual void DestroyRenderState(FFastGeoDestroyRenderStateContext* Context) override;
+	//~ End FFastGeoPrimitiveComponent interface
+
+	//~ Begin FastGeoMeshComponent interface
+	virtual UMaterialInterface* GetOverlayMaterial() const override;
+	virtual const TArray<TObjectPtr<UMaterialInterface>>& GetComponentMaterialSlotsOverlayMaterial() const override;
+	virtual void GetDefaultMaterialSlotsOverlayMaterial(TArray<TObjectPtr<UMaterialInterface>>& AssetMaterialSlotOverlayMaterials) const override;
+	//~ End FastGeoMeshComponent interface
+
+	//~ Begin FastGeoSkinnedMeshComponent interface
+	virtual FSkeletalMeshObject* CreateMeshObject();
+	virtual FSkinnedMeshSceneProxyDesc& GetSkinnedMeshSceneProxyDesc() = 0;
+	virtual const FSkinnedMeshSceneProxyDesc& GetSkinnedMeshSceneProxyDesc() const = 0;
+	virtual FPrimitiveSceneProxy* AllocateSceneProxy() = 0;
+	virtual void UpdateSkinning() = 0;
+	//~ End FastGeoSkinnedMeshComponent interface
+
+	FBoxSphereBounds CalcMeshBound(const FVector3f& InRootOffset, bool bInUsePhysicsAsset, const FTransform& InLocalToWorld) const;
+	USkinnedAsset* GetSkinnedAsset() const { return GetSkinnedMeshSceneProxyDesc().GetSkinnedAsset(); }
+	FSkeletalMeshObject* GetMeshObject() { return MeshObject; }
+	const FSkeletalMeshObject* GetMeshObject() const { return MeshObject; }
+	void DestroyMeshObject();
+
+	//~ Begin Unsupported
+	UPhysicsAsset* GetPhysicsAsset() const { return nullptr; }
+	int32 GetPredictedLODLevel() const { return 0; }
+	//~ End Unsupported
+
+protected:
+	// Persistent data
+	TArray<ESkinCacheUsage> SkinCacheUsage;
+	bool bOverrideMinLod : 1 = false;
+	bool bIncludeComponentLocationIntoBounds = false;
+	bool bHideSkin : 1 = false;	
+	int32 MinLodModel = 0;
+
+	// Transient
+	TArray<FSkelMeshComponentLODInfo> LODInfo;
+	FSkeletalMeshObject* MeshObject = nullptr;
+
+	friend class FSkinnedMeshComponentHelper;
+};
+
+class FASTGEOSTREAMING_API FFastGeoSkinnedMeshComponent : public FFastGeoSkinnedMeshComponentBase
+{
+public:
+	typedef FFastGeoSkinnedMeshComponentBase Super;
+
+	/** Static type identifier for this element class */
+	static const FFastGeoElementType Type;
+
+	FFastGeoSkinnedMeshComponent(int32 InComponentIndex = INDEX_NONE, FFastGeoElementType InType = Type);
+	virtual ~FFastGeoSkinnedMeshComponent() = default;
+	
+protected:
+	//~ Begin FFastGeoPrimitiveComponent interface
+	virtual FPrimitiveSceneProxyDesc& GetSceneProxyDesc() override { return SceneProxyDesc; }
+	virtual const FPrimitiveSceneProxyDesc& GetSceneProxyDesc() const override { return SceneProxyDesc; }
+	virtual void ApplyWorldTransform(const FTransform& InTransform) override;
+	//~ End FFastGeoPrimitiveComponent interface
+
+	//~ Begin FFastGeoSkinnedMeshComponentBase interface
+	virtual FSkinnedMeshSceneProxyDesc& GetSkinnedMeshSceneProxyDesc() override { return SceneProxyDesc; }
+	virtual const FSkinnedMeshSceneProxyDesc& GetSkinnedMeshSceneProxyDesc() const override { return SceneProxyDesc; }
+	virtual FPrimitiveSceneProxy* AllocateSceneProxy() override;
+	virtual void UpdateSkinning() override;
+	//~ End FFastGeoSkinnedMeshComponentBase interface	
+
+private:
+	friend class FSkinnedMeshComponentHelper;
+
+	FSkinnedMeshSceneProxyDesc SceneProxyDesc{};
+};
\ No newline at end of file
diff --git a/Engine/Plugins/Experimental/FastGeoStreaming/Source/FastGeoStreaming/Internal/FastGeoStaticMeshComponent.h b/Engine/Plugins/Experimental/FastGeoStreaming/Source/FastGeoStreaming/Internal/FastGeoStaticMeshComponent.h
new file mode 100644
--- /dev/null
+++ b/Engine/Plugins/Experimental/FastGeoStreaming/Source/FastGeoStreaming/Internal/FastGeoStaticMeshComponent.h
@@ -0,0 +1,171 @@
+// Copyright Epic Games, Inc. All Rights Reserved.
+
+#pragma once
+
+#include "FastGeoMeshComponent.h"
+#include "StaticMeshSceneProxyDesc.h"
+#include "FastGeoStaticMeshComponent.generated.h"
+
+class UBodySetup;
+class UMaterialInterface;
+class UNavCollisionBase;
+class FVertexFactoryType;
+struct FMaterialRelevance;
+struct FStaticMeshLODResources;
+struct FPrimitiveMaterialPropertyDescriptor;
+namespace ERHIFeatureLevel { enum Type : int; }
+namespace Nanite
+{
+	struct FResources;
+	class FNaniteResourcesHelper;
+}
+
+class FASTGEOSTREAMING_API FFastGeoStaticMeshComponentBase : public FFastGeoMeshComponent
+{
+public:
+	typedef FFastGeoMeshComponent Super;
+
+	/** Static type identifier for this element class */
+	static const FFastGeoElementType Type;
+
+	FFastGeoStaticMeshComponentBase(int32 InComponentIndex = INDEX_NONE, FFastGeoElementType InType = Type);
+	virtual ~FFastGeoStaticMeshComponentBase() = default;
+
+	const FCollisionResponseContainer& GetCollisionResponseToChannels() const;
+	UMaterialInterface* GetNaniteAuditMaterial(int32 MaterialIndex) const;
+	const Nanite::FResources* GetNaniteResources() const;
+	UStaticMesh* GetStaticMesh() const;
+	bool HasValidNaniteData() const;
+	bool UseNaniteOverrideMaterials(bool bDoingMaterialAudit) const;
+	bool ShouldCreateNaniteProxy(Nanite::FMaterialAudit* OutNaniteMaterials) const;
+	virtual UObject const* AdditionalStatObject() const override;
+	bool IsReverseCulling() const;
+	bool IsDisallowNanite() const;
+	bool IsForceDisableNanite() const;
+	bool IsForceNaniteForMasked() const;
+	int32 GetForcedLodModel() const;
+	bool GetOverrideMinLOD() const;
+	int32 GetMinLOD() const;
+	bool GetForceNaniteForMasked() const;
+#if WITH_EDITORONLY_DATA
+	bool IsDisplayNaniteFallbackMesh() const;
+#endif
+
+	//~ Begin FFastGeoComponent interface
+	virtual void Serialize(FArchive& Ar) override;
+	virtual void InitializeDynamicProperties() override;
+	virtual UBodySetup* GetBodySetup() const override;
+#if WITH_EDITOR
+	virtual void InitializeFromComponent(UActorComponent* Component) override;
+	virtual UClass* GetEditorProxyClass() const override;
+#endif
+	//~ End FFastGeoComponent interface
+
+	//~ Begin FFastGeoPrimitiveComponent interface
+	virtual bool IsNavigationRelevant() const override;
+	virtual FBox GetNavigationBounds() const override;
+	virtual void GetNavigationData(FNavigationRelevantData& Data) const override;
+	virtual bool DoCustomNavigableGeometryExport(FNavigableGeometryExport& GeomExport) const override;
+	virtual UMaterialInterface* GetMaterial(int32 MaterialIndex) const override;
+	virtual int32 GetNumMaterials() const override;
+	virtual void GetUsedMaterials(TArray<UMaterialInterface*>& OutMaterials, bool bGetDebugMaterials = false) const override;
+	//~ End FFastGeoPrimitiveComponent interface
+
+protected:
+	//~ Begin FFastGeoPrimitiveComponent interface
+#if WITH_EDITOR
+	virtual void InitializeSceneProxyDescFromComponent(UActorComponent* Component) override;
+	virtual void ResetSceneProxyDescUnsupportedProperties() override;
+#endif
+	virtual void InitializeSceneProxyDescDynamicProperties() override;
+	virtual void ApplyWorldTransform(const FTransform& InTransform) override;
+	virtual FPrimitiveSceneProxy* CreateSceneProxy(ESceneProxyCreationError* OutError = nullptr) override;
+	virtual UPhysicalMaterial* GetPhysicalMaterial() const override;
+	virtual void GetComplexPhysicalMaterials(TArray<UPhysicalMaterial*>& OutPhysMaterials, TArray<FPhysicalMaterialMaskParams>* OutPhysMaterialMasks) const override;
+	//~ End FFastGeoPrimitiveComponent interface
+
+	//~ Begin FFastGeoMeshComponent interface
+	virtual UMaterialInterface* GetOverlayMaterial() const override;
+	virtual const TArray<TObjectPtr<UMaterialInterface>>& GetComponentMaterialSlotsOverlayMaterial() const override;
+	virtual void GetDefaultMaterialSlotsOverlayMaterial(TArray<TObjectPtr<UMaterialInterface>>& OutMaterialSlotOverlayMaterials) const override;
+	//~ End FFastGeoMeshComponent interface
+
+	virtual FStaticMeshSceneProxyDesc& GetStaticMeshSceneProxyDesc() = 0;
+	virtual const FStaticMeshSceneProxyDesc& GetStaticMeshSceneProxyDesc() const = 0;
+	virtual FPrimitiveSceneProxy* CreateStaticMeshSceneProxy(const Nanite::FMaterialAudit& NaniteMaterials, bool bCreateNanite);
+	UMaterialInterface* GetMaterial(int32 MaterialIndex, bool bDoingNaniteMaterialAudit) const;
+	bool ShouldExportAsObstacle(const UNavCollisionBase& InNavCollision) const;
+#if !WITH_EDITOR && UE_WITH_PSO_PRECACHING
+	virtual void CollectPSOPrecacheData(const FPSOPrecacheParams& BasePrecachePSOParams, FMaterialInterfacePSOPrecacheParamsList& OutParams);
+#endif
+
+protected:
+	bool bUseDefaultCollision = false;
+	// Temporary until LODData support
+	struct FFastGeoStaticMeshComponentLODInfo
+	{
+		FColorVertexBuffer* OverrideVertexColors;
+	};
+	TArray<FFastGeoStaticMeshComponentLODInfo> LODData;
+
+private:
+	friend class FMeshComponentHelper;
+	friend class FPrimitiveComponentHelper;
+	friend class FStaticMeshComponentHelper;
+	friend class FInstancedStaticMeshComponentHelper;
+	friend class Nanite::FNaniteResourcesHelper;
+};
+
+class FASTGEOSTREAMING_API FFastGeoStaticMeshComponent : public FFastGeoStaticMeshComponentBase
+{
+public:
+	typedef FFastGeoStaticMeshComponentBase Super;
+
+	/** Static type identifier for this element class */
+	static const FFastGeoElementType Type;
+
+	FFastGeoStaticMeshComponent(int32 InComponentIndex = INDEX_NONE, FFastGeoElementType InType = Type);
+	virtual ~FFastGeoStaticMeshComponent() = default;
+
+protected:
+	//~ Begin FFastGeoPrimitiveComponent interface
+	virtual FPrimitiveSceneProxyDesc& GetSceneProxyDesc() override { return SceneProxyDesc; }
+	virtual const FPrimitiveSceneProxyDesc& GetSceneProxyDesc() const override { return SceneProxyDesc; }
+	//~ End FFastGeoPrimitiveComponent interface
+
+	//~ Begin FFastGeoStaticMeshComponentBase interface
+	virtual FStaticMeshSceneProxyDesc& GetStaticMeshSceneProxyDesc() override { return SceneProxyDesc; }
+	virtual const FStaticMeshSceneProxyDesc& GetStaticMeshSceneProxyDesc() const override { return SceneProxyDesc; }
+	//~ End FFastGeoStaticMeshComponentBase interface
+
+private:
+	FStaticMeshSceneProxyDesc SceneProxyDesc{};
+};
+
+UCLASS()
+class UFastGeoStaticMeshComponentEditorProxy : public UFastGeoPrimitiveComponentEditorProxy, public IFastGeoEditorInterface<IStaticMeshComponent>
+{
+	GENERATED_BODY()
+
+#if WITH_EDITOR
+	typedef FFastGeoStaticMeshComponentBase ComponentType;
+
+public:
+	//~ Begin UFastGeoPrimitiveComponentEditorProxy interface
+	virtual void NotifyRenderStateChanged() override;
+	//~ End UFastGeoPrimitiveComponentEditorProxy interface
+
+protected:
+	//~ Begin IStaticMeshComponent interface
+	virtual void OnMeshRebuild(bool bRenderDataChanged) override;
+	virtual void PreStaticMeshCompilation() override;
+	virtual void PostStaticMeshCompilation() override;
+	virtual UStaticMesh* GetStaticMesh() const override;
+	
+	virtual IPrimitiveComponent* GetPrimitiveComponentInterface() override
+	{
+		return UFastGeoPrimitiveComponentEditorProxy::GetPrimitiveComponentInterface();
+	}
+	//~ End IStaticMeshComponent interface
+#endif
+};
\ No newline at end of file
diff --git a/Engine/Plugins/Experimental/FastGeoStreaming/Source/FastGeoStreaming/Internal/FastGeoWeakElement.h b/Engine/Plugins/Experimental/FastGeoStreaming/Source/FastGeoStreaming/Internal/FastGeoWeakElement.h
new file mode 100644
--- /dev/null
+++ b/Engine/Plugins/Experimental/FastGeoStreaming/Source/FastGeoStreaming/Internal/FastGeoWeakElement.h
@@ -0,0 +1,92 @@
+// Copyright Epic Games, Inc. All Rights Reserved.
+
+#pragma once
+
+#include "FastGeoContainer.h"
+#include "FastGeoComponentCluster.h"
+#include "FastGeoComponent.h"
+#include "UObject/WeakObjectPtrTemplates.h"
+
+class FWeakFastGeoComponentCluster
+{
+public:
+	FWeakFastGeoComponentCluster(FFastGeoComponentCluster* InComponentCluster = nullptr)
+	{
+		if (UFastGeoContainer* Container = InComponentCluster ? InComponentCluster->GetOwnerContainer() : nullptr)
+		{
+			ContainerWeak = Container;
+			ComponentClusterTypeID = InComponentCluster->GetTypeID();
+			ComponentClusterIndex = InComponentCluster->GetComponentClusterIndex();
+		}
+	}
+
+	FORCEINLINE bool operator==(const FWeakFastGeoComponentCluster& Other) const
+	{
+		return ContainerWeak == Other.ContainerWeak &&
+			ComponentClusterTypeID == Other.ComponentClusterTypeID &&
+			ComponentClusterIndex == Other.ComponentClusterIndex;
+	}
+
+	FORCEINLINE FFastGeoComponentCluster* Get() const
+	{
+		if (UFastGeoContainer* Container = ContainerWeak.Get())
+		{
+			return Container->GetComponentCluster(ComponentClusterTypeID, ComponentClusterIndex);
+		}
+		return nullptr;
+	}
+
+	template <class T>
+	FORCEINLINE T* Get() const
+	{
+		FFastGeoComponentCluster* ComponentCluster = Get();
+		return ComponentCluster ? ComponentCluster->CastTo<T>() : nullptr;
+	}
+
+private:
+	TWeakObjectPtr<UFastGeoContainer> ContainerWeak;
+	uint32 ComponentClusterTypeID = INDEX_NONE;
+	int32 ComponentClusterIndex = INDEX_NONE;
+};
+
+class FWeakFastGeoComponent
+{
+public:
+	FWeakFastGeoComponent(FFastGeoComponent* InComponent)
+	{
+		if (FFastGeoComponentCluster* ComponentCluster = InComponent ? InComponent->GetOwnerComponentCluster() : nullptr)
+		{
+			ComponentClusterWeak = ComponentCluster;
+			ComponentTypeID = InComponent->GetTypeID();
+			ComponentIndex = InComponent->GetComponentIndex();
+		}
+	}
+
+	FORCEINLINE bool operator==(const FWeakFastGeoComponent& Other) const
+	{
+		return ComponentClusterWeak == Other.ComponentClusterWeak &&
+			ComponentTypeID == Other.ComponentTypeID &&
+			ComponentIndex == Other.ComponentIndex;
+	}
+
+	FORCEINLINE FFastGeoComponent* Get() const
+	{
+		if (FFastGeoComponentCluster* ComponentCluster = ComponentClusterWeak.Get())
+		{
+			return ComponentCluster->GetComponent(ComponentTypeID, ComponentIndex);
+		}
+		return nullptr;
+	}
+
+	template <class T>
+	FORCEINLINE T* Get() const
+	{
+		FFastGeoComponent* Component = Get();
+		return Component ? Component->CastTo<T>() : nullptr;
+	}
+
+private:
+	FWeakFastGeoComponentCluster ComponentClusterWeak;
+	uint32 ComponentTypeID = INDEX_NONE;
+	int32 ComponentIndex = INDEX_NONE;
+};
diff --git a/Engine/Plugins/Experimental/FastGeoStreaming/Source/FastGeoStreaming/Internal/IFastGeoElement.h b/Engine/Plugins/Experimental/FastGeoStreaming/Source/FastGeoStreaming/Internal/IFastGeoElement.h
new file mode 100644
--- /dev/null
+++ b/Engine/Plugins/Experimental/FastGeoStreaming/Source/FastGeoStreaming/Internal/IFastGeoElement.h
@@ -0,0 +1,57 @@
+// Copyright Epic Games, Inc. All Rights Reserved.
+
+#pragma once
+
+#include "FastGeoElementType.h"
+
+class IFastGeoElement
+{
+public:
+	/** Default constructor */
+	IFastGeoElement(FFastGeoElementType InType) : ElementType(InType) {}
+	virtual ~IFastGeoElement() {}
+
+	/** Attempt to cast this element to another type if it is of that type. Returns nullptr if it fails */
+	template <typename T>
+	T* CastTo()
+	{
+		return ElementType.IsA(T::Type) ? StaticCast<T*>(this) : nullptr;
+	}
+	/** Attempt to cast this element to another type if it is of that type. Returns nullptr if it fails */
+	template <typename T>
+	const T* CastTo() const
+	{
+		return ElementType.IsA(T::Type) ? StaticCast<const T*>(this) : nullptr;
+	}
+	/** Cast this element to another type and return a reference. Will assert if the cast fails. */
+	template <typename T>
+	T& CastToRef()
+	{
+		check(ElementType.IsA(T::Type));
+		return *StaticCast<T*>(this);
+	}
+	/** Cast this element to another type and return a reference. Will assert if the cast fails. */
+	template <typename T>
+	const T& CastToRef() const
+	{
+		check(ElementType.IsA(T::Type));
+		return *StaticCast<const T*>(this);
+	}
+	/** Returns true if this element is of the specified type */
+	template <typename T>
+	bool IsA() const
+	{
+		return ElementType.IsA(T::Type);
+	}
+
+	/** Returns the element type unique ID */
+	uint32 GetTypeID() { return ElementType.ID; }
+
+protected:
+	/** Static type identifier for the base class */
+	static const FFastGeoElementType Type;
+
+	/** Element type identifier */
+	FFastGeoElementType ElementType;
+};
+
diff --git a/Engine/Plugins/Experimental/FastGeoStreaming/Source/FastGeoStreaming/Private/FastGeoComponent.h b/Engine/Plugins/Experimental/FastGeoStreaming/Source/FastGeoStreaming/Private/FastGeoComponent.h
deleted file mode 100644
diff --git a/Engine/Plugins/Experimental/FastGeoStreaming/Source/FastGeoStreaming/Private/FastGeoComponentCluster.h b/Engine/Plugins/Experimental/FastGeoStreaming/Source/FastGeoStreaming/Private/FastGeoComponentCluster.h
deleted file mode 100644
diff --git a/Engine/Plugins/Experimental/FastGeoStreaming/Source/FastGeoStreaming/Private/FastGeoContainer.h b/Engine/Plugins/Experimental/FastGeoStreaming/Source/FastGeoStreaming/Private/FastGeoContainer.h
deleted file mode 100644
diff --git a/Engine/Plugins/Experimental/FastGeoStreaming/Source/FastGeoStreaming/Private/FastGeoElementType.h b/Engine/Plugins/Experimental/FastGeoStreaming/Source/FastGeoStreaming/Private/FastGeoElementType.h
deleted file mode 100644
diff --git a/Engine/Plugins/Experimental/FastGeoStreaming/Source/FastGeoStreaming/Private/FastGeoHLOD.h b/Engine/Plugins/Experimental/FastGeoStreaming/Source/FastGeoStreaming/Private/FastGeoHLOD.h
deleted file mode 100644
diff --git a/Engine/Plugins/Experimental/FastGeoStreaming/Source/FastGeoStreaming/Private/FastGeoInstancedSkinnedMeshComponent.h b/Engine/Plugins/Experimental/FastGeoStreaming/Source/FastGeoStreaming/Private/FastGeoInstancedSkinnedMeshComponent.h
deleted file mode 100644
diff --git a/Engine/Plugins/Experimental/FastGeoStreaming/Source/FastGeoStreaming/Private/FastGeoInstancedStaticMeshComponent.h b/Engine/Plugins/Experimental/FastGeoStreaming/Source/FastGeoStreaming/Private/FastGeoInstancedStaticMeshComponent.h
deleted file mode 100644
diff --git a/Engine/Plugins/Experimental/FastGeoStreaming/Source/FastGeoStreaming/Private/FastGeoMeshComponent.h b/Engine/Plugins/Experimental/FastGeoStreaming/Source/FastGeoStreaming/Private/FastGeoMeshComponent.h
deleted file mode 100644
diff --git a/Engine/Plugins/Experimental/FastGeoStreaming/Source/FastGeoStreaming/Private/FastGeoPrimitiveComponent.h b/Engine/Plugins/Experimental/FastGeoStreaming/Source/FastGeoStreaming/Private/FastGeoPrimitiveComponent.h
deleted file mode 100644
diff --git a/Engine/Plugins/Experimental/FastGeoStreaming/Source/FastGeoStreaming/Private/FastGeoSkinnedMeshComponent.h b/Engine/Plugins/Experimental/FastGeoStreaming/Source/FastGeoStreaming/Private/FastGeoSkinnedMeshComponent.h
deleted file mode 100644
diff --git a/Engine/Plugins/Experimental/FastGeoStreaming/Source/FastGeoStreaming/Private/FastGeoStaticMeshComponent.h b/Engine/Plugins/Experimental/FastGeoStreaming/Source/FastGeoStreaming/Private/FastGeoStaticMeshComponent.h
deleted file mode 100644
diff --git a/Engine/Plugins/Experimental/FastGeoStreaming/Source/FastGeoStreaming/Private/FastGeoWeakElement.h b/Engine/Plugins/Experimental/FastGeoStreaming/Source/FastGeoStreaming/Private/FastGeoWeakElement.h
deleted file mode 100644
diff --git a/Engine/Plugins/Experimental/FastGeoStreaming/Source/FastGeoStreaming/Private/IFastGeoElement.h b/Engine/Plugins/Experimental/FastGeoStreaming/Source/FastGeoStreaming/Private/IFastGeoElement.h
deleted file mode 100644
