[FastGeo] Add missing unitialized variable fixes that got missed after reconciling previous review.

#rnx

diff --git a/Engine/Plugins/Experimental/FastGeoStreaming/Source/FastGeoStreaming/Internal/FastGeoPrimitiveComponent.h b/Engine/Plugins/Experimental/FastGeoStreaming/Source/FastGeoStreaming/Internal/FastGeoPrimitiveComponent.h
--- a/Engine/Plugins/Experimental/FastGeoStreaming/Source/FastGeoStreaming/Internal/FastGeoPrimitiveComponent.h
+++ b/Engine/Plugins/Experimental/FastGeoStreaming/Source/FastGeoStreaming/Internal/FastGeoPrimitiveComponent.h
@@ -189,8 +189,8 @@ protected:
 	uint8 bRasterizeAsFilledConvexVolume : 1 = false;
 	uint8 bCanEverAffectNavigation : 1 = false;
 	FCustomPrimitiveData CustomPrimitiveData;
-	TEnumAsByte<enum EDetailMode> DetailMode;
-	TEnumAsByte<EHasCustomNavigableGeometry::Type> bHasCustomNavigableGeometry;
+	TEnumAsByte<enum EDetailMode> DetailMode = EDetailMode::DM_Low;
+	TEnumAsByte<EHasCustomNavigableGeometry::Type> bHasCustomNavigableGeometry = EHasCustomNavigableGeometry::No;
 	FBodyInstance BodyInstance;
 	TArray<URuntimeVirtualTexture*> RuntimeVirtualTextures;
 
