[FastGeo] Avoid copying transient state across components.

#rnx

diff --git a/Engine/Plugins/Experimental/FastGeoStreaming/Source/FastGeoStreaming/Private/FastGeoPrimitiveComponent.cpp b/Engine/Plugins/Experimental/FastGeoStreaming/Source/FastGeoStreaming/Private/FastGeoPrimitiveComponent.cpp
--- a/Engine/Plugins/Experimental/FastGeoStreaming/Source/FastGeoStreaming/Private/FastGeoPrimitiveComponent.cpp
+++ b/Engine/Plugins/Experimental/FastGeoStreaming/Source/FastGeoStreaming/Private/FastGeoPrimitiveComponent.cpp
@@ -78,24 +78,22 @@ FFastGeoPrimitiveComponent::FFastGeoPrimitiveComponent(const FFastGeoPrimitiveCo
 	, CustomPrimitiveData(Other.CustomPrimitiveData)
 	, DetailMode(Other.DetailMode)
 	, bHasCustomNavigableGeometry(Other.bHasCustomNavigableGeometry)
-	, BodyInstance(Other.BodyInstance)
+	, BodyInstance()
 	, RuntimeVirtualTextures(Other.RuntimeVirtualTextures)
-	, BodyInstanceOwner(Other.BodyInstanceOwner)
-	, PrimitiveSceneData(Other.PrimitiveSceneData)
-	, AsyncTermBodyPayload(Other.AsyncTermBodyPayload)
-	, ProxyState(Other.ProxyState)
-	, bRenderStateDirty(Other.bRenderStateDirty)
+	, BodyInstanceOwner()
+	, PrimitiveSceneData()
+	, AsyncTermBodyPayload()
+	, ProxyState(EProxyCreationState::None)
+	, bRenderStateDirty(false)
 #if !WITH_EDITOR && UE_WITH_PSO_PRECACHING
-	, MaterialPSOPrecacheRequestIDs(Other.MaterialPSOPrecacheRequestIDs)
-	, LatestPSOPrecacheJobSet(Other.LatestPSOPrecacheJobSet)
+	, MaterialPSOPrecacheRequestIDs()
+	, LatestPSOPrecacheJobSetCompleted(0)
+	, LatestPSOPrecacheJobSet(0)
+	, bPSOPrecacheCalled(false)
+	, bPSOPrecacheRequired(false)
+	, PSOPrecacheRequestPriority(EPSOPrecachePriority::Medium)
 #endif
 {
-#if !WITH_EDITOR && UE_WITH_PSO_PRECACHING
-	LatestPSOPrecacheJobSetCompleted.store(Other.LatestPSOPrecacheJobSetCompleted.load());
-	bPSOPrecacheCalled.store(Other.bPSOPrecacheCalled.load());
-	bPSOPrecacheRequired.store(Other.bPSOPrecacheRequired.load());
-	PSOPrecacheRequestPriority.store(Other.PSOPrecacheRequestPriority.load());
-#endif
 }
 
 void FFastGeoPrimitiveComponent::Serialize(FArchive& Ar)
