[FastGeo] Move checks for async physics state creation into create/destroy functions to relax the constraint if no physics state is required.

#rnx

diff --git a/Engine/Plugins/Experimental/FastGeoStreaming/Source/FastGeoStreaming/Private/FastGeoContainer.cpp b/Engine/Plugins/Experimental/FastGeoStreaming/Source/FastGeoStreaming/Private/FastGeoContainer.cpp
--- a/Engine/Plugins/Experimental/FastGeoStreaming/Source/FastGeoStreaming/Private/FastGeoContainer.cpp
+++ b/Engine/Plugins/Experimental/FastGeoStreaming/Source/FastGeoStreaming/Private/FastGeoContainer.cpp
@@ -741,6 +741,8 @@ void UFastGeoContainer::OnCreatePhysicsStateBegin_GameThread()
 
 	if (!CollisionComponents.IsEmpty())
 	{
+		check(FPhysScene::SupportsAsyncPhysicsStateCreation());
+
 		PendingCreate.PhysicsState.ComponentsToProcess = &CollisionComponents;
 		check(PendingCreate.PhysicsState.TotalNumProcessed == 0);
 
@@ -758,6 +760,8 @@ void UFastGeoContainer::OnDestroyPhysicsStateBegin_GameThread()
 
 	if (!CollisionComponents.IsEmpty())
 	{
+		check(FPhysScene::SupportsAsyncPhysicsStateDestruction());
+
 		PendingDestroy.PhysicsState.ComponentsToProcess = &CollisionComponents;
 		check(PendingDestroy.PhysicsState.TotalNumProcessed == 0);
 
diff --git a/Engine/Plugins/Experimental/FastGeoStreaming/Source/FastGeoStreaming/Private/FastGeoProceduralISMComponent.cpp b/Engine/Plugins/Experimental/FastGeoStreaming/Source/FastGeoStreaming/Private/FastGeoProceduralISMComponent.cpp
--- a/Engine/Plugins/Experimental/FastGeoStreaming/Source/FastGeoStreaming/Private/FastGeoProceduralISMComponent.cpp
+++ b/Engine/Plugins/Experimental/FastGeoStreaming/Source/FastGeoStreaming/Private/FastGeoProceduralISMComponent.cpp
@@ -82,6 +82,7 @@ void FFastGeoProceduralISMComponent::ResetSceneProxyDescUnsupportedProperties()
 	SceneProxyDesc.bHasSelectedInstances = false;
 	SceneProxyDesc.bAffectDynamicIndirectLighting = false;
 	SceneProxyDesc.bAffectDistanceFieldLighting = false;
+	SceneProxyDesc.bCollisionEnabled = false;
 }
 #endif
 
@@ -91,6 +92,7 @@ void FFastGeoProceduralISMComponent::InitializeSceneProxyDescDynamicProperties()
 
 	SceneProxyDesc.bAffectDynamicIndirectLighting = false;
 	SceneProxyDesc.bAffectDistanceFieldLighting = false;
+	SceneProxyDesc.bCollisionEnabled = false;
 }
 
 void FFastGeoProceduralISMComponent::ApplyWorldTransform(const FTransform& InTransform)
diff --git a/Engine/Plugins/Experimental/FastGeoStreaming/Source/FastGeoStreaming/Private/FastGeoWorldSubsystem.cpp b/Engine/Plugins/Experimental/FastGeoStreaming/Source/FastGeoStreaming/Private/FastGeoWorldSubsystem.cpp
--- a/Engine/Plugins/Experimental/FastGeoStreaming/Source/FastGeoStreaming/Private/FastGeoWorldSubsystem.cpp
+++ b/Engine/Plugins/Experimental/FastGeoStreaming/Source/FastGeoStreaming/Private/FastGeoWorldSubsystem.cpp
@@ -510,13 +510,11 @@ void UFastGeoWorldSubsystem::ProcessAsyncRenderStateJobs(bool bWaitForCompletion
 
 void UFastGeoWorldSubsystem::PushAsyncCreatePhysicsStateJobs(UFastGeoContainer* FastGeo)
 {
-	check(FPhysScene::SupportsAsyncPhysicsStateCreation());
 	FastGeo->OnCreatePhysicsStateBegin_GameThread();
 }
 
 void UFastGeoWorldSubsystem::PushAsyncDestroyPhysicsStateJobs(UFastGeoContainer* FastGeo)
 {
-	check(FPhysScene::SupportsAsyncPhysicsStateDestruction());
 	FastGeo->OnDestroyPhysicsStateBegin_GameThread();
 }
 
