[FastGeo] Fixed FFastGeoPrimitiveComponent BodyInstance not properly copied (regression caused by CL45312313)

diff --git a/Engine/Plugins/Experimental/FastGeoStreaming/Source/FastGeoStreaming/Private/FastGeoPrimitiveComponent.cpp b/Engine/Plugins/Experimental/FastGeoStreaming/Source/FastGeoStreaming/Private/FastGeoPrimitiveComponent.cpp
--- a/Engine/Plugins/Experimental/FastGeoStreaming/Source/FastGeoStreaming/Private/FastGeoPrimitiveComponent.cpp
+++ b/Engine/Plugins/Experimental/FastGeoStreaming/Source/FastGeoStreaming/Private/FastGeoPrimitiveComponent.cpp
@@ -78,7 +78,7 @@ FFastGeoPrimitiveComponent::FFastGeoPrimitiveComponent(const FFastGeoPrimitiveCo
 	, CustomPrimitiveData(Other.CustomPrimitiveData)
 	, DetailMode(Other.DetailMode)
 	, bHasCustomNavigableGeometry(Other.bHasCustomNavigableGeometry)
-	, BodyInstance()
+	, BodyInstance(Other.BodyInstance)
 	, RuntimeVirtualTextures(Other.RuntimeVirtualTextures)
 	, BodyInstanceOwner()
 	, PrimitiveSceneData()
