[FastGeo] Fix crash on mac caused by TArray relocation of FFastGeoPrimitiveComponent which contains a FRWLock.

diff --git a/Engine/Plugins/Experimental/FastGeoStreaming/Source/FastGeoStreaming/Internal/FastGeoPrimitiveComponent.h b/Engine/Plugins/Experimental/FastGeoStreaming/Source/FastGeoStreaming/Internal/FastGeoPrimitiveComponent.h
--- a/Engine/Plugins/Experimental/FastGeoStreaming/Source/FastGeoStreaming/Internal/FastGeoPrimitiveComponent.h
+++ b/Engine/Plugins/Experimental/FastGeoStreaming/Source/FastGeoStreaming/Internal/FastGeoPrimitiveComponent.h
@@ -237,7 +237,7 @@ protected:
 	static_assert((int)EPSOPrecachePriority::Highest < 1 << 2);
 #endif
 
-	TDontCopy<FRWLock> Lock;
+	TUniquePtr<FRWLock> Lock;
 
 	friend class FFastGeoComponentCluster;
 	friend class FPrimitiveComponentHelper;
diff --git a/Engine/Plugins/Experimental/FastGeoStreaming/Source/FastGeoStreaming/Private/FastGeoComponentCluster.cpp b/Engine/Plugins/Experimental/FastGeoStreaming/Source/FastGeoStreaming/Private/FastGeoComponentCluster.cpp
--- a/Engine/Plugins/Experimental/FastGeoStreaming/Source/FastGeoStreaming/Private/FastGeoComponentCluster.cpp
+++ b/Engine/Plugins/Experimental/FastGeoStreaming/Source/FastGeoStreaming/Private/FastGeoComponentCluster.cpp
@@ -110,7 +110,7 @@ void FFastGeoComponentCluster::UpdateVisibility()
 	TArray<FFastGeoPrimitiveComponent*> HideComponents;
 	ForEachComponent<FFastGeoPrimitiveComponent>([&ShowComponents, &HideComponents](FFastGeoPrimitiveComponent& Component)
 	{
-		FWriteScopeLock WriteLock(Component.Lock.Get());
+		FWriteScopeLock WriteLock(*Component.Lock.Get());
 		const bool bOldIsDrawnInGame = Component.IsDrawnInGame();
 		Component.UpdateVisibility();
 		const bool bIsDrawnInGame = Component.IsDrawnInGame();
@@ -141,7 +141,7 @@ void FFastGeoComponentCluster::ForceUpdateVisibility(const TArray<FFastGeoPrimit
 	TArray<FFastGeoPrimitiveComponent*> HideComponents;
 	for (FFastGeoPrimitiveComponent* Component : Components)
 	{
-		FWriteScopeLock WriteLock(Component->Lock.Get());
+		FWriteScopeLock WriteLock(*Component->Lock.Get());
 		Component->UpdateVisibility();
 		const bool bIsDrawnInGame = Component->IsDrawnInGame();
 		if (FPrimitiveSceneProxy* SceneProxy = Component->GetSceneProxy())
@@ -179,7 +179,7 @@ void FFastGeoComponentCluster::UpdateVisibility_Internal(TArray<FFastGeoPrimitiv
 			{
 				for (FFastGeoPrimitiveComponent* Component : InOutComponents)
 				{
-					FReadScopeLock ReadLock(Component->Lock.Get());
+					FReadScopeLock ReadLock(*Component->Lock.Get());
 					if (FPrimitiveSceneProxy* Proxy = Component->GetSceneProxy())
 					{
 						// Test whether the primitive was added to the scene (or is pending)
diff --git a/Engine/Plugins/Experimental/FastGeoStreaming/Source/FastGeoStreaming/Private/FastGeoPrimitiveComponent.cpp b/Engine/Plugins/Experimental/FastGeoStreaming/Source/FastGeoStreaming/Private/FastGeoPrimitiveComponent.cpp
--- a/Engine/Plugins/Experimental/FastGeoStreaming/Source/FastGeoStreaming/Private/FastGeoPrimitiveComponent.cpp
+++ b/Engine/Plugins/Experimental/FastGeoStreaming/Source/FastGeoStreaming/Private/FastGeoPrimitiveComponent.cpp
@@ -61,6 +61,7 @@ FFastGeoPrimitiveComponent::FFastGeoPrimitiveComponent(int32 InComponentIndex, F
 	: Super(InComponentIndex, InType)
 	, LocalBounds(ForceInit)
 	, WorldBounds(ForceInit)
+	, Lock(MakeUnique<FRWLock>())
 {
 }
 
@@ -93,6 +94,7 @@ FFastGeoPrimitiveComponent::FFastGeoPrimitiveComponent(const FFastGeoPrimitiveCo
 	, bPSOPrecacheRequired(false)
 	, PSOPrecacheRequestPriority(EPSOPrecachePriority::Medium)
 #endif
+	, Lock(MakeUnique<FRWLock>())
 {
 }
 
@@ -344,7 +346,7 @@ FPrimitiveSceneDesc FFastGeoPrimitiveComponent::BuildSceneDesc()
 
 void FFastGeoPrimitiveComponent::CreateRenderState(FRegisterComponentContext* Context)
 {
-	FWriteScopeLock WriteLock(Lock.Get());
+	FWriteScopeLock WriteLock(*Lock.Get());
 	ProxyState = EProxyCreationState::Creating;
 	bRenderStateDirty = false;
 
@@ -407,7 +409,7 @@ void FFastGeoPrimitiveComponent::FFastGeoDestroyRenderStateContext::DestroyProxy
 
 void FFastGeoPrimitiveComponent::DestroyRenderState(FFastGeoDestroyRenderStateContext* Context)
 {
-	FWriteScopeLock WriteLock(Lock.Get());
+	FWriteScopeLock WriteLock(*Lock.Get());
 	if (GetSceneProxy())
 	{
 		check(ProxyState == EProxyCreationState::Created);
