[FastGeo] Reusable Settings Asset for FastGeo transformers. An Instanced asset ("EmbeddedSettings") will replace the old, now deprecated setting properties and be auto-upgraded. Another property can be used ("Settings") to point to a standalone, reusable asset. Using a standalone asset disables the embedded one. A button can be used to automate this in the property grid under the embedded settings.

[RN] minor plugin/fastgeostreaming
#tests CitySample cook, FastGeoSimple functional test, x64, Full Editor Test preflight

diff --git a/Engine/Plugins/Experimental/FastGeoStreaming/FastGeoStreaming.uplugin b/Engine/Plugins/Experimental/FastGeoStreaming/FastGeoStreaming.uplugin
--- a/Engine/Plugins/Experimental/FastGeoStreaming/FastGeoStreaming.uplugin
+++ b/Engine/Plugins/Experimental/FastGeoStreaming/FastGeoStreaming.uplugin
@@ -17,6 +17,11 @@
 			"Name": "FastGeoStreaming",
 			"Type": "Runtime",
 			"LoadingPhase": "PreDefault"
+		},
+		{
+			"Name": "FastGeoStreamingEditor",
+			"Type": "Editor",
+			"LoadingPhase": "Default"
 		}
 	],
 	"Plugins": [
diff --git a/Engine/Plugins/Experimental/FastGeoStreaming/Source/FastGeoStreaming/Private/FastGeoWorldPartitionRuntimeCellTransformer.cpp b/Engine/Plugins/Experimental/FastGeoStreaming/Source/FastGeoStreaming/Private/FastGeoWorldPartitionRuntimeCellTransformer.cpp
--- a/Engine/Plugins/Experimental/FastGeoStreaming/Source/FastGeoStreaming/Private/FastGeoWorldPartitionRuntimeCellTransformer.cpp
+++ b/Engine/Plugins/Experimental/FastGeoStreaming/Source/FastGeoStreaming/Private/FastGeoWorldPartitionRuntimeCellTransformer.cpp
@@ -133,11 +133,15 @@ UFastGeoWorldPartitionRuntimeCellTransformer::UFastGeoWorldPartitionRuntimeCellT
 	{
 		USelection::SelectionChangedEvent.AddUObject(this, &UFastGeoWorldPartitionRuntimeCellTransformer::OnSelectionChanged);
 	}
+
+	EmbeddedSettings = CreateDefaultSubobject<UFastGeoTransformerSettings>(TEXT("EmbeddedSettings"));
 #endif
 }
 
 void UFastGeoWorldPartitionRuntimeCellTransformer::Serialize(FArchive& Ar)
 {
+	Ar.UsingCustomVersion(FFortniteMainBranchObjectVersion::GUID);
+
 	Super::Serialize(Ar);
 
 #if WITH_EDITOR
@@ -145,6 +149,24 @@ void UFastGeoWorldPartitionRuntimeCellTransformer::Serialize(FArchive& Ar)
 	{
 		Ar << bDebugMode;
 	}
+
+	if (Ar.IsLoading())
+	{
+		// Convert setting properties to embedded settings asset
+		if (Ar.CustomVer(FFortniteMainBranchObjectVersion::GUID) < FFortniteMainBranchObjectVersion::FastGeoTransformerSettingAssets)
+		{
+			EmbeddedSettings->AllowedActorClasses                   = AllowedActorClasses_DEPRECATED;
+			EmbeddedSettings->AllowedExactActorClasses              = AllowedExactActorClasses_DEPRECATED;
+			EmbeddedSettings->AllowedComponentClasses               = AllowedComponentClasses_DEPRECATED;
+			EmbeddedSettings->AllowedExactComponentClasses          = AllowedExactComponentClasses_DEPRECATED;
+			EmbeddedSettings->DisallowedActorClasses                = DisallowedActorClasses_DEPRECATED;
+			EmbeddedSettings->DisallowedExactActorClasses           = DisallowedExactActorClasses_DEPRECATED;
+			EmbeddedSettings->DisallowedComponentClasses            = DisallowedComponentClasses_DEPRECATED;
+			EmbeddedSettings->DisallowedExactComponentClasses       = DisallowedExactComponentClasses_DEPRECATED;
+			EmbeddedSettings->IgnoredRemainingComponentClasses      = IgnoredRemainingComponentClasses_DEPRECATED;
+			EmbeddedSettings->IgnoredRemainingExactComponentClasses = IgnoredRemainingExactComponentClasses_DEPRECATED;
+		}
+	}
 #endif
 }
 
@@ -373,17 +395,18 @@ void UFastGeoWorldPartitionRuntimeCellTransformer::ForEachIgnoredComponentClass(
 {
 	Super::ForEachIgnoredComponentClass(Func);
 
-	for (const TSubclassOf<UActorComponent>& IgnoredComponentClass : IgnoredRemainingComponentClasses)
+	const UFastGeoTransformerSettings* EffectiveSettings = GetSettings();
+
+	for (const TSubclassOf<UActorComponent>& IgnoredRemainingComponentClass : EffectiveSettings->IgnoredRemainingComponentClasses)
 	{
-		if (!Func(IgnoredComponentClass))
+		if (!Func(IgnoredRemainingComponentClass))
 		{
 			return;
 		}
 	}
-
-	for (const TSubclassOf<UActorComponent>& IgnoredComponentClass : BuiltinIgnoredRemainingComponentClasses)
+	for (const TSubclassOf<UActorComponent>& BuiltinIgnoredRemainingComponentClass : BuiltinIgnoredRemainingComponentClasses)
 	{
-		if (!Func(IgnoredComponentClass))
+		if (!Func(BuiltinIgnoredRemainingComponentClass))
 		{
 			return;
 		}
@@ -394,17 +417,18 @@ void UFastGeoWorldPartitionRuntimeCellTransformer::ForEachIgnoredExactComponentC
 {
 	Super::ForEachIgnoredExactComponentClass(Func);
 
-	for (const TSubclassOf<UActorComponent>& IgnoredComponentClass : IgnoredRemainingExactComponentClasses)
+	const UFastGeoTransformerSettings* EffectiveSettings = GetSettings();
+
+	for (const TSubclassOf<UActorComponent>& IgnoredRemainingExactComponentClass : EffectiveSettings->IgnoredRemainingExactComponentClasses)
 	{
-		if (!Func(IgnoredComponentClass))
+		if (!Func(IgnoredRemainingExactComponentClass))
 		{
 			return;
 		}
 	}
-
-	for (const TSubclassOf<UActorComponent>& IgnoredComponentClass : BuiltinIgnoredRemainingExactComponentClasses)
+	for (const TSubclassOf<UActorComponent>& BuiltinIgnoredRemainingExactComponentClass : BuiltinIgnoredRemainingExactComponentClasses)
 	{
-		if (!Func(IgnoredComponentClass))
+		if (!Func(BuiltinIgnoredRemainingExactComponentClass))
 		{
 			return;
 		}
@@ -631,8 +655,10 @@ bool UFastGeoWorldPartitionRuntimeCellTransformer::IsBlueprintActorWithLogic(AAc
 
 FFastGeoTransformResult UFastGeoWorldPartitionRuntimeCellTransformer::IsAllowedActorClass(AActor* InActor) const
 {
+	const UFastGeoTransformerSettings* EffectiveSettings = GetSettings();
 	UClass* ActorClass = InActor->GetClass();
-	for (TSubclassOf<AActor> DisallowedActorClass : DisallowedActorClasses)
+
+	for (TSubclassOf<AActor> DisallowedActorClass : EffectiveSettings->DisallowedActorClasses)
 	{
 		if (ActorClass->IsChildOf(DisallowedActorClass))
 		{
@@ -648,15 +674,15 @@ FFastGeoTransformResult UFastGeoWorldPartitionRuntimeCellTransformer::IsAllowedA
 		}
 	}
 
-	for (TSubclassOf<AActor> DisallowedActorClass : DisallowedExactActorClasses)
+	for (TSubclassOf<AActor> DisallowedExactActorClass : EffectiveSettings->DisallowedExactActorClasses)
 	{
-		if (ActorClass == DisallowedActorClass)
+		if (ActorClass == DisallowedExactActorClass)
 		{
-			return { EFastGeoTransform::Reject, [InActor, DisallowedActorClass] { return FString::Printf(TEXT("Actor %s class is a disallowed exact class (%s)"), *InActor->GetName(), *DisallowedActorClass->GetName()); } };
+			return { EFastGeoTransform::Reject, [InActor, DisallowedExactActorClass] { return FString::Printf(TEXT("Actor %s class is a disallowed exact class (%s)"), *InActor->GetName(), *DisallowedExactActorClass->GetName()); } };
 		}
 	}
 
-	for (TSubclassOf<AActor> AllowedActorClass : AllowedActorClasses)
+	for (TSubclassOf<AActor> AllowedActorClass : EffectiveSettings->AllowedActorClasses)
 	{
 		if (ActorClass->IsChildOf(AllowedActorClass))
 		{
@@ -664,17 +690,17 @@ FFastGeoTransformResult UFastGeoWorldPartitionRuntimeCellTransformer::IsAllowedA
 		}
 	}
 
-	for (TSubclassOf<AActor> AllowedActorClass : BuiltinAllowedActorClasses)
+	for (TSubclassOf<AActor> BuiltinAllowedActorClass : BuiltinAllowedActorClasses)
 	{
-		if (ActorClass->IsChildOf(AllowedActorClass))
+		if (ActorClass->IsChildOf(BuiltinAllowedActorClass))
 		{
 			return EFastGeoTransform::Allow;
 		}
 	}
 
-	for (TSubclassOf<AActor> AllowedActorClass : AllowedExactActorClasses)
+	for (TSubclassOf<AActor> AllowedExactActorClass : EffectiveSettings->AllowedExactActorClasses)
 	{
-		if (ActorClass->IsChildOf(AllowedActorClass))
+		if (ActorClass == AllowedExactActorClass)
 		{
 			return EFastGeoTransform::Allow;
 		}
@@ -797,6 +823,8 @@ FFastGeoTransformResult UFastGeoWorldPartitionRuntimeCellTransformer::CanTransfo
 
 FFastGeoTransformResult UFastGeoWorldPartitionRuntimeCellTransformer::IsAllowedComponentClass(UActorComponent* InComponent) const
 {
+	const UFastGeoTransformerSettings* EffectiveSettings = GetSettings();
+
 	FFastGeoElementType FastGeoComponentType = FastGeo::GetFastGeoComponentType(InComponent->GetClass());
 	if (!FastGeoComponentType.IsValid())
 	{
@@ -804,31 +832,31 @@ FFastGeoTransformResult UFastGeoWorldPartitionRuntimeCellTransformer::IsAllowedC
 	}
 
 	UClass* ComponentClass = InComponent->GetClass();
-	for (TSubclassOf<UActorComponent> DisallowedComponentClass : DisallowedComponentClasses)
+	for (TSubclassOf<UActorComponent> DisallowedComponentClass : EffectiveSettings->DisallowedComponentClasses)
 	{
 		if (ComponentClass->IsChildOf(DisallowedComponentClass))
 		{
 			return { EFastGeoTransform::Reject, [InComponent, DisallowedComponentClass] { return FString::Printf(TEXT("Component %s class is child of a disallowed class (%s)"), *FastGeo::GetComponentShortName(InComponent), *DisallowedComponentClass->GetName()); } };
 		}
 	}
-
-	for (TSubclassOf<UActorComponent> DisallowedComponentClass : BuiltinDisallowedComponentClasses)
+	
+	for (TSubclassOf<UActorComponent> BuiltinDisallowedComponentClass : BuiltinDisallowedComponentClasses)
 	{
-		if (ComponentClass->IsChildOf(DisallowedComponentClass))
+		if (ComponentClass->IsChildOf(BuiltinDisallowedComponentClass))
 		{
-			return { EFastGeoTransform::Reject, [InComponent, DisallowedComponentClass] { return FString::Printf(TEXT("Component %s class is child of a built-in disallowed class (%s)"), *FastGeo::GetComponentShortName(InComponent), *DisallowedComponentClass->GetName()); } };
+			return { EFastGeoTransform::Reject, [InComponent, BuiltinDisallowedComponentClass] { return FString::Printf(TEXT("Component %s class is child of a built-in disallowed class (%s)"), *FastGeo::GetComponentShortName(InComponent), *BuiltinDisallowedComponentClass->GetName()); } };
 		}
 	}
 
-	for (TSubclassOf<UActorComponent> DisallowedComponentClass : DisallowedExactComponentClasses)
+	for (TSubclassOf<UActorComponent> DisallowedExactComponentClass : EffectiveSettings->DisallowedExactComponentClasses)
 	{
-		if (ComponentClass == DisallowedComponentClass)
+		if (ComponentClass == DisallowedExactComponentClass)
 		{
-			return { EFastGeoTransform::Reject, [InComponent, DisallowedComponentClass] { return FString::Printf(TEXT("Component %s class is a disallowed exact class (%s)"), *FastGeo::GetComponentShortName(InComponent), *DisallowedComponentClass->GetName()); } };
+			return { EFastGeoTransform::Reject, [InComponent, DisallowedExactComponentClass] { return FString::Printf(TEXT("Component %s class is a disallowed exact class (%s)"), *FastGeo::GetComponentShortName(InComponent), *DisallowedExactComponentClass->GetName()); } };
 		}
 	}
 
-	for (TSubclassOf<UActorComponent> AllowedComponentClass : AllowedComponentClasses)
+	for (TSubclassOf<UActorComponent> AllowedComponentClass : EffectiveSettings->AllowedComponentClasses)
 	{
 		if (ComponentClass->IsChildOf(AllowedComponentClass))
 		{
@@ -836,17 +864,17 @@ FFastGeoTransformResult UFastGeoWorldPartitionRuntimeCellTransformer::IsAllowedC
 		}
 	}
 
-	for (TSubclassOf<UActorComponent> AllowedComponentClass : BuiltinAllowedComponentClasses)
+	for (TSubclassOf<UActorComponent> BuiltinAllowedComponentClass : BuiltinAllowedComponentClasses)
 	{
-		if (ComponentClass->IsChildOf(AllowedComponentClass))
+		if (ComponentClass->IsChildOf(BuiltinAllowedComponentClass))
 		{
 			return EFastGeoTransform::Allow;
 		}
 	}
-
-	for (TSubclassOf<UActorComponent> AllowedComponentClass : AllowedExactComponentClasses)
+	
+	for (TSubclassOf<UActorComponent> AllowedExactComponentClass : EffectiveSettings->AllowedExactComponentClasses)
 	{
-		if (ComponentClass == AllowedComponentClass)
+		if (ComponentClass == AllowedExactComponentClass)
 		{
 			return EFastGeoTransform::Allow;
 		}
diff --git a/Engine/Plugins/Experimental/FastGeoStreaming/Source/FastGeoStreaming/Public/FastGeoWorldPartitionRuntimeCellTransformer.h b/Engine/Plugins/Experimental/FastGeoStreaming/Source/FastGeoStreaming/Public/FastGeoWorldPartitionRuntimeCellTransformer.h
--- a/Engine/Plugins/Experimental/FastGeoStreaming/Source/FastGeoStreaming/Public/FastGeoWorldPartitionRuntimeCellTransformer.h
+++ b/Engine/Plugins/Experimental/FastGeoStreaming/Source/FastGeoStreaming/Public/FastGeoWorldPartitionRuntimeCellTransformer.h
@@ -7,6 +7,7 @@
 #include "Templates/SubclassOf.h"
 #include "HAL/IConsoleManager.h"
 #endif 
+#include "Engine/DataAsset.h"
 #include "FastGeoWorldPartitionRuntimeCellTransformer.generated.h"
 
 class UActorComponent;
@@ -40,6 +41,63 @@
 	EFastGeoTransform TransformResult;
 };
 
+// Empty type used to present a Button in the property grid - not actual data
+USTRUCT()
+struct FFastGeoConversionButton
+{
+	GENERATED_BODY()
+};
+
+UCLASS(MinimalAPI, BlueprintType, EditInlineNew, CollapseCategories)
+class UFastGeoTransformerSettings : public UDataAsset
+{
+	GENERATED_BODY()
+	
+public:
+	/** Allowed actor classes (recursive) to convert to FastGeo */
+	UPROPERTY(EditAnywhere, Category="Settings")
+	TArray<TSubclassOf<AActor>> AllowedActorClasses;
+
+	/** Allowed actor classes (exact) to convert to FastGeo */
+	UPROPERTY(EditAnywhere, Category="Settings")
+	TArray<TSubclassOf<AActor>> AllowedExactActorClasses;
+
+	/** Allowed component classes (recursive) to convert to FastGeo */
+	UPROPERTY(EditAnywhere, Category="Settings")
+	TArray<TSubclassOf<UActorComponent>> AllowedComponentClasses;
+
+	/** Allowed component classes (exact) to convert to FastGeo */
+	UPROPERTY(EditAnywhere, Category="Settings")
+	TArray<TSubclassOf<UActorComponent>> AllowedExactComponentClasses;
+
+	/** Disallowed actor classes (recursive) to convert to FastGeo */
+	UPROPERTY(EditAnywhere, Category="Settings")
+	TArray<TSubclassOf<AActor>> DisallowedActorClasses;
+
+	/** Disallowed actor classes (exact) to convert to FastGeo */
+	UPROPERTY(EditAnywhere, Category="Settings")
+	TArray<TSubclassOf<AActor>> DisallowedExactActorClasses;
+
+	/** Disallowed component classes (recursive) to convert to FastGeo */
+	UPROPERTY(EditAnywhere, Category="Settings")
+	TArray<TSubclassOf<UActorComponent>> DisallowedComponentClasses;
+
+	/** Disallowed component classes (exact) to convert to FastGeo */
+	UPROPERTY(EditAnywhere, Category="Settings")
+	TArray<TSubclassOf<UActorComponent>> DisallowedExactComponentClasses;
+
+	/** Component classes (recursive) to ignore when deciding to destroy a converted actor to FastGeo */
+	UPROPERTY(EditAnywhere, Category="Settings")
+	TArray<TSubclassOf<UActorComponent>> IgnoredRemainingComponentClasses;
+
+	/** Component classes (exact) to ignore when deciding to destroy a converted actor to FastGeo */
+	UPROPERTY(EditAnywhere, Category="Settings")
+	TArray<TSubclassOf<UActorComponent>> IgnoredRemainingExactComponentClasses;
+
+	UPROPERTY(VisibleAnywhere, Category = "Settings", meta = (DisplayName = "Convert settings to reusable Asset"))
+	FFastGeoConversionButton FastGeoConversionButton;
+};
+
 UCLASS(Config = FastGeoStreaming, DefaultConfig)
 class FASTGEOSTREAMING_API UFastGeoWorldPartitionRuntimeCellTransformer : public UWorldPartitionRuntimeCellTransformer
 {
@@ -104,55 +162,58 @@
 	void GatherTransformableActors(const TArray<AActor*>& InActors, const ULevel* InLevel, TMap<AActor*, FTransformableActor>& OutTransformableActors, FTransformationStats& OutStats);
 #endif
 
+public:
 #if WITH_EDITORONLY_DATA
+	UPROPERTY(EditAnywhere, Category = "Settings")
+	TObjectPtr<UFastGeoTransformerSettings> Settings = nullptr;
+#endif
+
 protected:
-	/** When enabled, PIE will log the reason why actors/components are not transformable */
-	UPROPERTY(EditAnywhere, Transient, SkipSerialization, Category = FastGeo)
+#if WITH_EDITORONLY_DATA
+	UPROPERTY(EditAnywhere, Instanced, Meta = (EditCondition = "Settings == nullptr", Tooltip = "Fallback only available when no Settings asset is used."), Category = "Settings")
+	TObjectPtr<UFastGeoTransformerSettings> EmbeddedSettings;
+	
+	const UFastGeoTransformerSettings* GetSettings() const
+	{
+		return (Settings != nullptr) ? Settings : EmbeddedSettings;
+	}
+
+	UPROPERTY(EditAnywhere, Transient, SkipSerialization, Category = "Debug")
 	bool bDebugMode;
 
-	/** When enabled, will log the reason why selected actors/components are not transformable. Works in both PIE and Editor. */
-	UPROPERTY(EditAnywhere, Transient, SkipSerialization, Category = FastGeo)
+	UPROPERTY(EditAnywhere, Transient, SkipSerialization, Category = "Debug")
 	bool bDebugModeOnSelection;
+	
+	// The following properties are all deprecated in favor of the separate Settings property pointing to a reusable settings asset
+	UPROPERTY(meta = (DeprecatedProperty))
+	TArray<TSubclassOf<AActor>> AllowedActorClasses_DEPRECATED;
 
-	/** Allowed actor classes (recursive) to convert to FastGeo */
-	UPROPERTY(EditAnywhere, Category = FastGeo)
-	TArray<TSubclassOf<AActor>> AllowedActorClasses;
+	UPROPERTY(meta = (DeprecatedProperty))
+	TArray<TSubclassOf<AActor>> AllowedExactActorClasses_DEPRECATED;
 
-	/** Allowed actor classes (exact) to convert to FastGeo */
-	UPROPERTY(EditAnywhere, Category = FastGeo)
-	TArray<TSubclassOf<AActor>> AllowedExactActorClasses;
-
-	/** Allowed component classes (recursive) to convert to FastGeo */
-	UPROPERTY(EditAnywhere, Category = FastGeo)
-	TArray<TSubclassOf<UActorComponent>> AllowedComponentClasses;
-
-	/** Allowed component classes (exact) to convert to FastGeo */
-	UPROPERTY(EditAnywhere, Category = FastGeo)
-	TArray<TSubclassOf<UActorComponent>> AllowedExactComponentClasses;
-
-	/** Disallowed actor classes (recursive) to convert to FastGeo */
-	UPROPERTY(EditAnywhere, Category = FastGeo)
-	TArray<TSubclassOf<AActor>> DisallowedActorClasses;
-
-	/** Disallowed actor classes (exact) to convert to FastGeo */
-	UPROPERTY(EditAnywhere, Category = FastGeo)
-	TArray<TSubclassOf<AActor>> DisallowedExactActorClasses;
-
-	/** Disallowed component classes (recursive) to convert to FastGeo */
-	UPROPERTY(EditAnywhere, Category = FastGeo)
-	TArray<TSubclassOf<UActorComponent>> DisallowedComponentClasses;
-
-	/** Disallowed component classes (exact) to convert to FastGeo */
-	UPROPERTY(EditAnywhere, Category = FastGeo)
-	TArray<TSubclassOf<UActorComponent>> DisallowedExactComponentClasses;
-
-	/** Component classes (recursive) to ignore when deciding to destroy a converted actor to FastGeo */
-	UPROPERTY(EditAnywhere, Category = FastGeo)
-	TArray<TSubclassOf<UActorComponent>> IgnoredRemainingComponentClasses;
-
-	/** Component classes (exact) to ignore when deciding to destroy a converted actor to FastGeo */
-	UPROPERTY(EditAnywhere, Category = FastGeo)
-	TArray<TSubclassOf<UActorComponent>> IgnoredRemainingExactComponentClasses;
+	UPROPERTY(meta = (DeprecatedProperty))
+	TArray<TSubclassOf<UActorComponent>> AllowedComponentClasses_DEPRECATED;
+ 
+	UPROPERTY(meta = (DeprecatedProperty))
+	TArray<TSubclassOf<UActorComponent>> AllowedExactComponentClasses_DEPRECATED;
+ 
+	UPROPERTY(meta = (DeprecatedProperty))
+	TArray<TSubclassOf<AActor>> DisallowedActorClasses_DEPRECATED;
+ 
+	UPROPERTY(meta = (DeprecatedProperty))
+	TArray<TSubclassOf<AActor>> DisallowedExactActorClasses_DEPRECATED;
+ 
+	UPROPERTY(meta = (DeprecatedProperty))
+	TArray<TSubclassOf<UActorComponent>> DisallowedComponentClasses_DEPRECATED;
+ 
+	UPROPERTY(meta = (DeprecatedProperty))
+	TArray<TSubclassOf<UActorComponent>> DisallowedExactComponentClasses_DEPRECATED;
+ 
+	UPROPERTY(meta = (DeprecatedProperty))
+	TArray<TSubclassOf<UActorComponent>> IgnoredRemainingComponentClasses_DEPRECATED;
+ 
+	UPROPERTY(meta = (DeprecatedProperty))
+	TArray<TSubclassOf<UActorComponent>> IgnoredRemainingExactComponentClasses_DEPRECATED;
 #endif
 
 #if WITH_EDITOR
diff --git a/Engine/Plugins/Experimental/FastGeoStreaming/Source/FastGeoStreamingEditor/FastGeoStreamingEditor.Build.cs b/Engine/Plugins/Experimental/FastGeoStreaming/Source/FastGeoStreamingEditor/FastGeoStreamingEditor.Build.cs
new file mode 100644
--- /dev/null
+++ b/Engine/Plugins/Experimental/FastGeoStreaming/Source/FastGeoStreamingEditor/FastGeoStreamingEditor.Build.cs
@@ -0,0 +1,24 @@
+// Copyright Epic Games, Inc. All Rights Reserved.
+namespace UnrealBuildTool.Rules
+{
+	public class FastGeoStreamingEditor : ModuleRules
+	{
+		public FastGeoStreamingEditor(ReadOnlyTargetRules Target) : base(Target)
+		{
+			PrivateDependencyModuleNames.AddRange(
+				new string[]
+				{
+					"AssetDefinition",
+					"AssetTools",
+					"Core",
+					"CoreUObject",
+					"Engine",
+					"FastGeoStreaming",
+					"Slate",
+					"SlateCore",
+					"UnrealEd"
+				}
+			);
+		}
+	}
+}
\ No newline at end of file
diff --git a/Engine/Plugins/Experimental/FastGeoStreaming/Source/FastGeoStreamingEditor/Private/FastGeoAssetDefinitions.cpp b/Engine/Plugins/Experimental/FastGeoStreaming/Source/FastGeoStreamingEditor/Private/FastGeoAssetDefinitions.cpp
new file mode 100644
--- /dev/null
+++ b/Engine/Plugins/Experimental/FastGeoStreaming/Source/FastGeoStreamingEditor/Private/FastGeoAssetDefinitions.cpp
@@ -0,0 +1,36 @@
+// Copyright Epic Games, Inc. All Rights Reserved.
+
+#include "FastGeoAssetDefinitions.h"
+#include "FastGeoWorldPartitionRuntimeCellTransformer.h"
+
+// This file is required so that UFastGeoTransformerSettings appears in Asset creation dialogs
+
+#if WITH_EDITOR
+
+FText UAssetDefinition_FastGeoTransformerSettings::GetAssetDisplayName() const
+{
+	return NSLOCTEXT("AssetTypeActions", "AssetTypeActions_FastGeoTransformerSettings", "FastGeo Transformer Settings");
+}
+
+FLinearColor UAssetDefinition_FastGeoTransformerSettings::GetAssetColor() const
+{
+	return FLinearColor(FColor(12, 65, 12));
+}
+
+TSoftClassPtr<UObject> UAssetDefinition_FastGeoTransformerSettings::GetAssetClass() const
+{
+	return UFastGeoTransformerSettings::StaticClass();
+}
+
+FAssetSupportResponse UAssetDefinition_FastGeoTransformerSettings::CanLocalize(const FAssetData& InAsset) const
+{
+	return FAssetSupportResponse::NotSupported();
+}
+
+TConstArrayView<FAssetCategoryPath> UAssetDefinition_FastGeoTransformerSettings::GetAssetCategories() const
+{
+	static const FAssetCategoryPath Categories[] = { EAssetCategoryPaths::World };
+	return Categories;
+}
+
+#endif
\ No newline at end of file
diff --git a/Engine/Plugins/Experimental/FastGeoStreaming/Source/FastGeoStreamingEditor/Private/FastGeoAssetDefinitions.h b/Engine/Plugins/Experimental/FastGeoStreaming/Source/FastGeoStreamingEditor/Private/FastGeoAssetDefinitions.h
new file mode 100644
--- /dev/null
+++ b/Engine/Plugins/Experimental/FastGeoStreaming/Source/FastGeoStreamingEditor/Private/FastGeoAssetDefinitions.h
@@ -0,0 +1,25 @@
+// Copyright Epic Games, Inc. All Rights Reserved.
+
+#pragma once
+
+#include "AssetDefinitionDefault.h"
+#include "FastGeoAssetDefinitions.generated.h"
+
+#if WITH_EDITOR
+
+UCLASS()
+class UAssetDefinition_FastGeoTransformerSettings : public UAssetDefinitionDefault
+{
+	GENERATED_BODY()
+	
+public:
+	// UAssetDefinition Begin
+	virtual FText GetAssetDisplayName() const override;
+	virtual FLinearColor GetAssetColor() const override;
+	virtual TSoftClassPtr<UObject> GetAssetClass() const override;
+	virtual FAssetSupportResponse CanLocalize(const FAssetData& InAsset) const override;
+	virtual TConstArrayView<FAssetCategoryPath> GetAssetCategories() const override;
+	// UAssetDefinition End
+};
+
+#endif
\ No newline at end of file
diff --git a/Engine/Plugins/Experimental/FastGeoStreaming/Source/FastGeoStreamingEditor/Private/FastGeoDetailsCustomization.cpp b/Engine/Plugins/Experimental/FastGeoStreaming/Source/FastGeoStreamingEditor/Private/FastGeoDetailsCustomization.cpp
new file mode 100644
--- /dev/null
+++ b/Engine/Plugins/Experimental/FastGeoStreaming/Source/FastGeoStreamingEditor/Private/FastGeoDetailsCustomization.cpp
@@ -0,0 +1,94 @@
+// Copyright Epic Games, Inc. All Rights Reserved.
+
+#include "FastGeoDetailsCustomization.h"
+#include "AssetToolsModule.h"
+#include "DetailCategoryBuilder.h"
+#include "DetailLayoutBuilder.h"
+#include "DetailWidgetRow.h"
+#include "FastGeoFactory.h"
+#include "FastGeoWorldPartitionRuntimeCellTransformer.h"
+#include "ScopedTransaction.h"
+#include "Modules/ModuleManager.h"
+#include "Widgets/Input/SButton.h"
+#include "Widgets/Text/STextBlock.h"
+
+#define LOCTEXT_NAMESPACE "FastGeoDetailsCustomization"
+
+TSharedRef<IPropertyTypeCustomization> FConvertFastGeoSettingsAssetButtonCustomization::MakeInstance()
+{
+	return MakeShareable(new FConvertFastGeoSettingsAssetButtonCustomization);
+}
+
+void FConvertFastGeoSettingsAssetButtonCustomization::CustomizeHeader(TSharedRef<IPropertyHandle> PropertyHandle, FDetailWidgetRow& HeaderRow, IPropertyTypeCustomizationUtils& CustomizationUtils)
+{
+	bool bCreateConvertButton = false;
+
+	// Only offer a conversion button when this special property is found on a UFastGeoTransformerSettings, which itself is owned directly by a UFastGeoWorldPartitionRuntimeCellTransformer
+	// which means the UFastGeoTransformerSettings is the instanced EmbeddedSettings member, which we would like to make users convert to the separate Settings member
+	TArray<UObject*> Objects;
+	PropertyHandle->GetOuterObjects(Objects);
+	TWeakObjectPtr<UFastGeoTransformerSettings> TransformerSettings;
+	TWeakObjectPtr<UFastGeoWorldPartitionRuntimeCellTransformer> Transformer;
+
+	if (Objects.Num() == 1)
+	{
+		// Check that our special UFastGeoConversionButton property is indeed used in the UFastGeoTransformerSettings
+		TransformerSettings = Cast<UFastGeoTransformerSettings>(Objects[0]);
+		if (TransformerSettings != nullptr)
+		{
+			// And then check that we are not editing a standalone UFastGeoTransformerSettings asset, but rather the instanced EmbeddedSettings asset
+			// (we could check for the name of the property if there are ever more than one instanced UFastGeoTransformerSettings)
+			Transformer = Cast<UFastGeoWorldPartitionRuntimeCellTransformer>(TransformerSettings->GetOuter());
+			if (Transformer != nullptr)
+			{
+				bCreateConvertButton = true;
+			}
+		}
+	}
+
+	if (!bCreateConvertButton)
+	{
+		// Skipping here means the whole property row is removed, which is what we want
+		return;
+	}
+
+	HeaderRow.NameContent()
+		[
+			PropertyHandle->CreatePropertyNameWidget()
+		]
+		.ValueContent()
+		[
+			SNew(SButton)
+				.VAlign(VAlign_Center)
+				.HAlign(HAlign_Center)
+				.OnClicked_Lambda([Transformer, TransformerSettings]()
+					{
+						if (!Transformer.IsValid() || !TransformerSettings.IsValid())
+						{
+							return FReply::Handled();
+						}
+						// Support undo
+						const FScopedTransaction Transaction(LOCTEXT("ConvertFastGeoSettings", "Convert FastGeo Settings to standalone asset"));
+						// Create a new settings asset, then populate it with the current settings, then install the new asset in the transformer to replace these embedded settings
+						FAssetToolsModule& AssetToolsModule = FModuleManager::Get().LoadModuleChecked<FAssetToolsModule>("AssetTools");
+						UFastGeoFactory* Factory = NewObject<UFastGeoFactory>();
+						Factory->InitialSettings = TransformerSettings.Get(); // Set up the factory for a copy, since there are no AssetTools API to create an Asset with Dialog AND copy from an existing object
+						UFastGeoTransformerSettings* SettingsAsset = Cast<UFastGeoTransformerSettings>(IAssetTools::Get().CreateAssetWithDialog(UFastGeoTransformerSettings::StaticClass(), Factory));
+						if (SettingsAsset)
+						{
+							// Install the new Asset in the Transformer
+							Transformer->Settings = SettingsAsset;
+							Transformer->GetPackage()->SetDirtyFlag(true);
+						}
+						return FReply::Handled();
+					})
+				[
+					SNew(STextBlock)
+						.Text(LOCTEXT("ConvertFastGeoSettingToAssetButton", "Convert to Asset"))
+						.ToolTipText(LOCTEXT("ConvertFastGeoSettingToAsset_ToolTip", "Convert these settings to a separate reusable asset"))
+						.Font(IDetailLayoutBuilder::GetDetailFont())
+				]
+		];
+}
+
+#undef LOCTEXT_NAMESPACE
\ No newline at end of file
diff --git a/Engine/Plugins/Experimental/FastGeoStreaming/Source/FastGeoStreamingEditor/Private/FastGeoDetailsCustomization.h b/Engine/Plugins/Experimental/FastGeoStreaming/Source/FastGeoStreamingEditor/Private/FastGeoDetailsCustomization.h
new file mode 100644
--- /dev/null
+++ b/Engine/Plugins/Experimental/FastGeoStreaming/Source/FastGeoStreamingEditor/Private/FastGeoDetailsCustomization.h
@@ -0,0 +1,19 @@
+// Copyright Epic Games, Inc. All Rights Reserved.
+
+#pragma once
+
+#include "CoreMinimal.h"
+#include "SlateFwd.h"
+#include "IPropertyTypeCustomization.h"
+
+// Unfortunately, because UFastGeoTransformerSettings are Instanced in the case where we want the convert button to be available, and because IDetailCustomization currently doesn't work for Instanced properties,
+// we customize a virtual property instead
+class FConvertFastGeoSettingsAssetButtonCustomization : public IPropertyTypeCustomization
+{
+public:
+	static TSharedRef<IPropertyTypeCustomization> MakeInstance();
+	
+	// IDetailCustomization interface
+	virtual void CustomizeHeader(TSharedRef<IPropertyHandle> PropertyHandle, FDetailWidgetRow& HeaderRow, IPropertyTypeCustomizationUtils& CustomizationUtils) override;
+	void CustomizeChildren(TSharedRef<IPropertyHandle> PropertyHandle, IDetailChildrenBuilder& ChildBuilder, IPropertyTypeCustomizationUtils& CustomizationUtils) override {}
+};
\ No newline at end of file
diff --git a/Engine/Plugins/Experimental/FastGeoStreaming/Source/FastGeoStreamingEditor/Private/FastGeoFactory.cpp b/Engine/Plugins/Experimental/FastGeoStreaming/Source/FastGeoStreamingEditor/Private/FastGeoFactory.cpp
new file mode 100644
--- /dev/null
+++ b/Engine/Plugins/Experimental/FastGeoStreaming/Source/FastGeoStreamingEditor/Private/FastGeoFactory.cpp
@@ -0,0 +1,29 @@
+// Copyright Epic Games, Inc. All Rights Reserved.
+
+#include "FastGeoFactory.h"
+#include "FastGeoWorldPartitionRuntimeCellTransformer.h"
+
+#if WITH_EDITORONLY_DATA
+
+#include UE_INLINE_GENERATED_CPP_BY_NAME(FastGeoFactory)
+
+UFastGeoFactory::UFastGeoFactory(const FObjectInitializer& ObjectInitializer)
+	: Super(ObjectInitializer)
+{
+	SupportedClass = UFastGeoTransformerSettings::StaticClass();
+	bCreateNew = true;
+	bEditorImport = false;
+	bEditAfterNew = true;
+}
+
+UObject* UFastGeoFactory::FactoryCreateNew(UClass* Class, UObject* InParent, FName InName, EObjectFlags Flags, UObject* Context, FFeedbackContext* Warn)
+{
+	if (InitialSettings != nullptr)
+	{
+		return NewObject<UFastGeoTransformerSettings>(InParent, InName, Flags, InitialSettings);
+	}
+	
+	return NewObject<UFastGeoTransformerSettings>(InParent, InName, Flags);
+}
+
+#endif
\ No newline at end of file
diff --git a/Engine/Plugins/Experimental/FastGeoStreaming/Source/FastGeoStreamingEditor/Private/FastGeoFactory.h b/Engine/Plugins/Experimental/FastGeoStreaming/Source/FastGeoStreamingEditor/Private/FastGeoFactory.h
new file mode 100644
--- /dev/null
+++ b/Engine/Plugins/Experimental/FastGeoStreaming/Source/FastGeoStreamingEditor/Private/FastGeoFactory.h
@@ -0,0 +1,24 @@
+// Copyright Epic Games, Inc. All Rights Reserved.
+
+#pragma once
+
+#include "CoreMinimal.h"
+#include "Factories/Factory.h"
+#include "FastGeoFactory.generated.h"
+
+#if WITH_EDITOR
+
+UCLASS(hidecategories = Object, MinimalAPI)
+class UFastGeoFactory : public UFactory
+{
+	GENERATED_UCLASS_BODY()
+
+	// Set this in order for the factory to duplicate an object
+	UPROPERTY()
+	TObjectPtr<class UFastGeoTransformerSettings> InitialSettings;
+	
+	// Begin UFactory Interface
+	virtual UObject* FactoryCreateNew(UClass* Class, UObject* InParent, FName Name, EObjectFlags Flags, UObject* Context, FFeedbackContext* Warn) override;
+	// End UFactory Interface	
+};
+#endif
\ No newline at end of file
diff --git a/Engine/Plugins/Experimental/FastGeoStreaming/Source/FastGeoStreamingEditor/Private/FastGeoStreamingEditorModule.cpp b/Engine/Plugins/Experimental/FastGeoStreaming/Source/FastGeoStreamingEditor/Private/FastGeoStreamingEditorModule.cpp
new file mode 100644
--- /dev/null
+++ b/Engine/Plugins/Experimental/FastGeoStreaming/Source/FastGeoStreamingEditor/Private/FastGeoStreamingEditorModule.cpp
@@ -0,0 +1,29 @@
+// Copyright Epic Games, Inc. All Rights Reserved.
+
+#include "FastGeoStreamingEditorModule.h"
+#include "FastGeoDetailsCustomization.h"
+#include "FastGeoWorldPartitionRuntimeCellTransformer.h"
+#include "PropertyEditorModule.h"
+#include "Modules/ModuleManager.h"
+
+IMPLEMENT_MODULE(FFastGeoStreamingEditorModule, FastGeoStreamingEditor);
+
+void FFastGeoStreamingEditorModule::StartupModule()
+{
+	// Register property grid customization
+	if (FPropertyEditorModule* const PropertyModule = FModuleManager::GetModulePtr<FPropertyEditorModule>("PropertyEditor"))
+	{
+		PropertyModule->RegisterCustomPropertyTypeLayout(FFastGeoConversionButton::StaticStruct()->GetFName(), FOnGetPropertyTypeCustomizationInstance::CreateStatic(&FConvertFastGeoSettingsAssetButtonCustomization::MakeInstance));
+	}
+}
+
+void FFastGeoStreamingEditorModule::ShutdownModule()
+{
+	if (FModuleManager::Get().IsModuleLoaded(TEXT("PropertyEditor")))
+	{
+		if (FPropertyEditorModule* const PropertyModule = FModuleManager::GetModulePtr<FPropertyEditorModule>("PropertyEditor"))
+		{
+			PropertyModule->UnregisterCustomPropertyTypeLayout(FFastGeoConversionButton::StaticStruct()->GetFName());
+		}
+	}
+}
\ No newline at end of file
diff --git a/Engine/Plugins/Experimental/FastGeoStreaming/Source/FastGeoStreamingEditor/Private/FastGeoStreamingEditorModule.h b/Engine/Plugins/Experimental/FastGeoStreaming/Source/FastGeoStreamingEditor/Private/FastGeoStreamingEditorModule.h
new file mode 100644
--- /dev/null
+++ b/Engine/Plugins/Experimental/FastGeoStreaming/Source/FastGeoStreamingEditor/Private/FastGeoStreamingEditorModule.h
@@ -0,0 +1,16 @@
+// Copyright Epic Games, Inc. All Rights Reserved.
+
+#pragma once
+
+#include "Delegates/Delegate.h"
+#include "Modules/ModuleInterface.h"
+
+class FFastGeoPrimitiveComponent;
+class UWorld;
+
+class FFastGeoStreamingEditorModule : public IModuleInterface
+{
+public:
+	virtual void StartupModule() override;
+	virtual void ShutdownModule() override;
+};
\ No newline at end of file
diff --git a/Engine/Source/Runtime/Core/Public/UObject/FortniteMainBranchObjectVersions.inl b/Engine/Source/Runtime/Core/Public/UObject/FortniteMainBranchObjectVersions.inl
--- a/Engine/Source/Runtime/Core/Public/UObject/FortniteMainBranchObjectVersions.inl
+++ b/Engine/Source/Runtime/Core/Public/UObject/FortniteMainBranchObjectVersions.inl
@@ -688,5 +688,8 @@ FORTNITE_MAIN_VERSION(DynamicMeshAttributesMorphTargets, 224)
 // Introduce landscape advanced weight blending 
 FORTNITE_MAIN_VERSION(LandscapeAdvancedWeightBlending, 225)
 
+// Add support for FastGeo transformer setting assets
+FORTNITE_MAIN_VERSION(FastGeoTransformerSettingAssets, 226)
+
 // -----<new versions can be added above this line>-------------------------------------------------
 #undef FORTNITE_MAIN_VERSION
\ No newline at end of file
