[FastGeo] Updated all TSubclassOf<> properties of UFastGeoTransformerSettings to use meta = (AllowAbstract) so they can be filtered with abstract classes.

diff --git a/Engine/Plugins/Experimental/FastGeoStreaming/Source/FastGeoStreaming/Public/FastGeoWorldPartitionRuntimeCellTransformer.h b/Engine/Plugins/Experimental/FastGeoStreaming/Source/FastGeoStreaming/Public/FastGeoWorldPartitionRuntimeCellTransformer.h
--- a/Engine/Plugins/Experimental/FastGeoStreaming/Source/FastGeoStreaming/Public/FastGeoWorldPartitionRuntimeCellTransformer.h
+++ b/Engine/Plugins/Experimental/FastGeoStreaming/Source/FastGeoStreaming/Public/FastGeoWorldPartitionRuntimeCellTransformer.h
@@ -55,43 +55,43 @@ class UFastGeoTransformerSettings : public UDataAsset
 	
 public:
 	/** Allowed actor classes (recursive) to convert to FastGeo */
-	UPROPERTY(EditAnywhere, Category="Settings")
+	UPROPERTY(EditAnywhere, Category="Settings", meta= (AllowAbstract = "true"))
 	TArray<TSubclassOf<AActor>> AllowedActorClasses;
 
 	/** Allowed actor classes (exact) to convert to FastGeo */
-	UPROPERTY(EditAnywhere, Category="Settings")
+	UPROPERTY(EditAnywhere, Category="Settings", meta = (AllowAbstract = "true"))
 	TArray<TSubclassOf<AActor>> AllowedExactActorClasses;
 
 	/** Allowed component classes (recursive) to convert to FastGeo */
-	UPROPERTY(EditAnywhere, Category="Settings")
+	UPROPERTY(EditAnywhere, Category="Settings", meta = (AllowAbstract = "true"))
 	TArray<TSubclassOf<UActorComponent>> AllowedComponentClasses;
 
 	/** Allowed component classes (exact) to convert to FastGeo */
-	UPROPERTY(EditAnywhere, Category="Settings")
+	UPROPERTY(EditAnywhere, Category="Settings", meta = (AllowAbstract = "true"))
 	TArray<TSubclassOf<UActorComponent>> AllowedExactComponentClasses;
 
 	/** Disallowed actor classes (recursive) to convert to FastGeo */
-	UPROPERTY(EditAnywhere, Category="Settings")
+	UPROPERTY(EditAnywhere, Category="Settings", meta = (AllowAbstract = "true"))
 	TArray<TSubclassOf<AActor>> DisallowedActorClasses;
 
 	/** Disallowed actor classes (exact) to convert to FastGeo */
-	UPROPERTY(EditAnywhere, Category="Settings")
+	UPROPERTY(EditAnywhere, Category="Settings", meta = (AllowAbstract = "true"))
 	TArray<TSubclassOf<AActor>> DisallowedExactActorClasses;
 
 	/** Disallowed component classes (recursive) to convert to FastGeo */
-	UPROPERTY(EditAnywhere, Category="Settings")
+	UPROPERTY(EditAnywhere, Category="Settings", meta = (AllowAbstract = "true"))
 	TArray<TSubclassOf<UActorComponent>> DisallowedComponentClasses;
 
 	/** Disallowed component classes (exact) to convert to FastGeo */
-	UPROPERTY(EditAnywhere, Category="Settings")
+	UPROPERTY(EditAnywhere, Category="Settings", meta = (AllowAbstract = "true"))
 	TArray<TSubclassOf<UActorComponent>> DisallowedExactComponentClasses;
 
 	/** Component classes (recursive) to ignore when deciding to destroy a converted actor to FastGeo */
-	UPROPERTY(EditAnywhere, Category="Settings")
+	UPROPERTY(EditAnywhere, Category="Settings", meta = (AllowAbstract = "true"))
 	TArray<TSubclassOf<UActorComponent>> IgnoredRemainingComponentClasses;
 
 	/** Component classes (exact) to ignore when deciding to destroy a converted actor to FastGeo */
-	UPROPERTY(EditAnywhere, Category="Settings")
+	UPROPERTY(EditAnywhere, Category="Settings", meta = (AllowAbstract = "true"))
 	TArray<TSubclassOf<UActorComponent>> IgnoredRemainingExactComponentClasses;
 
 	UPROPERTY(VisibleAnywhere, Category = "Settings", meta = (DisplayName = "Convert settings to reusable Asset"))
