[FastGeo] Removed unused EProxyCreationState::Pending

diff --git a/Engine/Plugins/Experimental/FastGeoStreaming/Source/FastGeoStreaming/Internal/FastGeoPrimitiveComponent.h b/Engine/Plugins/Experimental/FastGeoStreaming/Source/FastGeoStreaming/Internal/FastGeoPrimitiveComponent.h
--- a/Engine/Plugins/Experimental/FastGeoStreaming/Source/FastGeoStreaming/Internal/FastGeoPrimitiveComponent.h
+++ b/Engine/Plugins/Experimental/FastGeoStreaming/Source/FastGeoStreaming/Internal/FastGeoPrimitiveComponent.h
@@ -203,7 +203,6 @@ protected:
 	enum class EProxyCreationState : uint8
 	{
 		None, // Constructed/Initialized
-		Pending, // AddedToWorld & proxy creation is pending
 		Creating, // Actively creating the proxy
 		Created, // Proxy is now created
 		Delayed // Proxy creation delayed (used when PSO precaching is not ready when creating proxy)
diff --git a/Engine/Plugins/Experimental/FastGeoStreaming/Source/FastGeoStreaming/Private/FastGeoPrimitiveComponent.cpp b/Engine/Plugins/Experimental/FastGeoStreaming/Source/FastGeoStreaming/Private/FastGeoPrimitiveComponent.cpp
--- a/Engine/Plugins/Experimental/FastGeoStreaming/Source/FastGeoStreaming/Private/FastGeoPrimitiveComponent.cpp
+++ b/Engine/Plugins/Experimental/FastGeoStreaming/Source/FastGeoStreaming/Private/FastGeoPrimitiveComponent.cpp
@@ -347,6 +347,7 @@ FPrimitiveSceneDesc FFastGeoPrimitiveComponent::BuildSceneDesc()
 void FFastGeoPrimitiveComponent::CreateRenderState(FRegisterComponentContext* Context)
 {
 	FWriteScopeLock WriteLock(*Lock.Get());
+	check(ProxyState == EProxyCreationState::None);
 	ProxyState = EProxyCreationState::Creating;
 	bRenderStateDirty = false;
 
@@ -418,7 +419,7 @@ void FFastGeoPrimitiveComponent::DestroyRenderState(FFastGeoDestroyRenderStateCo
 
 		PrimitiveSceneData.SceneProxy = nullptr;
 
-		ProxyState = EProxyCreationState::Pending;
+		ProxyState = EProxyCreationState::None;
 
 		bRenderStateDirty = false;
 
@@ -426,6 +427,11 @@ void FFastGeoPrimitiveComponent::DestroyRenderState(FFastGeoDestroyRenderStateCo
 		GetEditorProxy<EditorProxyType>()->NotifyRenderStateChanged();
 #endif
 	}
+	else
+	{
+		check(ProxyState == EProxyCreationState::None || ProxyState == EProxyCreationState::Delayed);
+		ProxyState = EProxyCreationState::None;
+	}
 }
 
 void FFastGeoPrimitiveComponent::OnAsyncCreatePhysicsState()
