[FastGeo] Cherry picked fix for FFastGeoSkinnedMeshComponentBase::DestroyMeshObject() being called outside of the GameThread.

#FailedValidations Preflight

diff --git a/Engine/Plugins/Experimental/FastGeoStreaming/Source/FastGeoStreaming/Private/FastGeoSkinnedMeshComponent.cpp b/Engine/Plugins/Experimental/FastGeoStreaming/Source/FastGeoStreaming/Private/FastGeoSkinnedMeshComponent.cpp
--- a/Engine/Plugins/Experimental/FastGeoStreaming/Source/FastGeoStreaming/Private/FastGeoSkinnedMeshComponent.cpp
+++ b/Engine/Plugins/Experimental/FastGeoStreaming/Source/FastGeoStreaming/Private/FastGeoSkinnedMeshComponent.cpp
@@ -177,8 +177,21 @@ void FFastGeoSkinnedMeshComponentBase::DestroyMeshObject()
 {
 	if (MeshObject)
 	{
-		MeshObject->ReleaseResources();
-		BeginCleanup(MeshObject);
+		if (IsInGameThread() || IsInParallelGameThread())
+		{
+			MeshObject->ReleaseResources();
+			BeginCleanup(MeshObject);
+		}
+		else
+		{
+			// MeshObject->ReleaseResources() needs to be called from the game thread
+			AsyncTask(ENamedThreads::GameThread, [MeshObject = MeshObject]()
+			{
+				MeshObject->ReleaseResources();
+				BeginCleanup(MeshObject);
+			});
+		}
+
 		MeshObject = nullptr;
 	}
 }
