[FastGeo] Fixed an issue where a streamed cell containing FastGeo content could attempt to re-create a render proxy while one was already being created.

diff --git a/Engine/Plugins/Experimental/FastGeoStreaming/Source/FastGeoStreaming/Internal/FastGeoContainer.h b/Engine/Plugins/Experimental/FastGeoStreaming/Source/FastGeoStreaming/Internal/FastGeoContainer.h
--- a/Engine/Plugins/Experimental/FastGeoStreaming/Source/FastGeoStreaming/Internal/FastGeoContainer.h
+++ b/Engine/Plugins/Experimental/FastGeoStreaming/Source/FastGeoStreaming/Internal/FastGeoContainer.h
@@ -292,8 +292,10 @@ private:
 	TArray<FNavigationElementHandle> NavigationElementHandles;
 
 	bool bIsRegistered = false;
+	int32 RegistrationEpoch = 0;
 
 	FRenderCommandFence DestroyFence;
 
 	friend class FFastGeoGatherFastGeoContainerAssetRefsArchive;
+	friend class UFastGeoWorldSubsystem;
 };
\ No newline at end of file
diff --git a/Engine/Plugins/Experimental/FastGeoStreaming/Source/FastGeoStreaming/Private/FastGeoContainer.cpp b/Engine/Plugins/Experimental/FastGeoStreaming/Source/FastGeoStreaming/Private/FastGeoContainer.cpp
--- a/Engine/Plugins/Experimental/FastGeoStreaming/Source/FastGeoStreaming/Private/FastGeoContainer.cpp
+++ b/Engine/Plugins/Experimental/FastGeoStreaming/Source/FastGeoStreaming/Private/FastGeoContainer.cpp
@@ -170,6 +170,7 @@ void UFastGeoContainer::Register()
 
 	if (!bIsRegistered)
 	{
+		++RegistrationEpoch;
 		PendingCreate.Reset();
 
 		check(GetWorld()->IsGameWorld());
@@ -242,10 +243,7 @@ void UFastGeoContainer::Unregister()
 
 			ComponentCluster.ForEachComponent<FFastGeoPrimitiveComponent>([this](FFastGeoPrimitiveComponent& Component)
 			{
-				if (Component.IsRenderStateCreated())
-				{
-					PendingDestroy.RenderState.ComponentsToProcess.Add(&Component);
-				}
+				PendingDestroy.RenderState.ComponentsToProcess.Add(&Component);
 			});
 		});
 
diff --git a/Engine/Plugins/Experimental/FastGeoStreaming/Source/FastGeoStreaming/Private/FastGeoWorldSubsystem.cpp b/Engine/Plugins/Experimental/FastGeoStreaming/Source/FastGeoStreaming/Private/FastGeoWorldSubsystem.cpp
--- a/Engine/Plugins/Experimental/FastGeoStreaming/Source/FastGeoStreaming/Private/FastGeoWorldSubsystem.cpp
+++ b/Engine/Plugins/Experimental/FastGeoStreaming/Source/FastGeoStreaming/Private/FastGeoWorldSubsystem.cpp
@@ -316,12 +316,7 @@ void UFastGeoWorldSubsystem::OnUpdateLevelStreaming()
 
 void UFastGeoWorldSubsystem::AddToComponentsPendingRecreate(FFastGeoPrimitiveComponent* InComponentPendingRecreate)
 {
-	ComponentsPendingRecreate.Add(InComponentPendingRecreate);
-}
-
-void UFastGeoWorldSubsystem::RemoveFromComponentsPendingRecreate(FFastGeoPrimitiveComponent* InComponentPendingRecreate)
-{
-	ComponentsPendingRecreate.Remove(InComponentPendingRecreate);
+	ComponentsPendingRecreate.Emplace(InComponentPendingRecreate);
 }
 
 void UFastGeoWorldSubsystem::ProcessPendingRecreate()
@@ -331,21 +326,22 @@ void UFastGeoWorldSubsystem::ProcessPendingRecreate()
 	CSV_CUSTOM_STAT(FastGeo, PendingRecreate, ComponentsPendingRecreate.Num(), ECsvCustomStatOp::Set);
 
 	int32 NumPrimitiveDelayed = 0;
-	for (FWeakFastGeoComponent WeakComponentPendingRecreate : ComponentsPendingRecreate)
+	for (const FPendingRecreateComponent& PendingRecreateComponent : ComponentsPendingRecreate)
 	{
-		if (FFastGeoPrimitiveComponent* ComponentPendingRecreate = WeakComponentPendingRecreate.Get<FFastGeoPrimitiveComponent>())
+		if (FFastGeoPrimitiveComponent* Component = PendingRecreateComponent.Component.Get<FFastGeoPrimitiveComponent>())
 		{
-#if CSV_PROFILER && !CSV_PROFILER_MINIMAL
-			if (ComponentPendingRecreate->IsRenderStateDelayed())
+			// Skip component if its container was unregistered or re-registered before we process it
+			if (Component->GetOwnerContainer()->IsRegistered() &&
+				(PendingRecreateComponent.RegistrationEpoch == Component->GetOwnerContainer()->RegistrationEpoch))
 			{
-				++NumPrimitiveDelayed;
-			}
+#if CSV_PROFILER && !CSV_PROFILER_MINIMAL
+				if (Component->IsRenderStateDelayed())
+				{
+					++NumPrimitiveDelayed;
+				}
 #endif
-			// Skip component if its container was unregistered before we process it
-			if (ComponentPendingRecreate->GetOwnerContainer()->IsRegistered())
-			{
-				ComponentPendingRecreate->DestroyRenderState(nullptr);
-				ComponentPendingRecreate->CreateRenderState(nullptr);
+				Component->DestroyRenderState(nullptr);
+				Component->CreateRenderState(nullptr);
 			}
 		}
 	}
diff --git a/Engine/Plugins/Experimental/FastGeoStreaming/Source/FastGeoStreaming/Private/FastGeoWorldSubsystem.h b/Engine/Plugins/Experimental/FastGeoStreaming/Source/FastGeoStreaming/Private/FastGeoWorldSubsystem.h
--- a/Engine/Plugins/Experimental/FastGeoStreaming/Source/FastGeoStreaming/Private/FastGeoWorldSubsystem.h
+++ b/Engine/Plugins/Experimental/FastGeoStreaming/Source/FastGeoStreaming/Private/FastGeoWorldSubsystem.h
@@ -37,7 +37,6 @@ public:
 	//~ End UTickableWorldSubsystem interface
 
 	void AddToComponentsPendingRecreate(FFastGeoPrimitiveComponent* InComponentPendingRecreate);
-	void RemoveFromComponentsPendingRecreate(FFastGeoPrimitiveComponent* InComponentPendingRecreate);
 	void ProcessPendingRecreate();
 
 	/**
@@ -99,7 +98,26 @@ private:
 	FDelegateHandle Handle_OnLevelBeginRemoveFromWorld;
 	FDelegateHandle Handle_OnForEachHLODObjectInCell;
 
-	TArray<FWeakFastGeoComponent> ComponentsPendingRecreate;
+	struct FPendingRecreateComponent
+	{
+		FPendingRecreateComponent(FFastGeoComponent* InComponent)
+			: Component(InComponent)
+		{
+			check(InComponent);
+			check(InComponent->GetOwnerContainer());
+			RegistrationEpoch = InComponent->GetOwnerContainer()->RegistrationEpoch;
+		}
+
+		FWeakFastGeoComponent Component;
+		int32 RegistrationEpoch = -1;
+
+		FORCEINLINE bool operator==(const FPendingRecreateComponent& Other) const
+		{
+			return Component == Other.Component && RegistrationEpoch == Other.RegistrationEpoch;
+		}
+	};
+
+	TArray<FPendingRecreateComponent> ComponentsPendingRecreate;
 	static bool bEnableDebugView;
 
 	FRWLock Lock;
diff --git a/Engine/Plugins/Experimental/FastGeoStreaming/Source/FastGeoStreaming/Public/FastGeoStreamingModule.h b/Engine/Plugins/Experimental/FastGeoStreaming/Source/FastGeoStreaming/Public/FastGeoStreamingModule.h
--- a/Engine/Plugins/Experimental/FastGeoStreaming/Source/FastGeoStreaming/Public/FastGeoStreamingModule.h
+++ b/Engine/Plugins/Experimental/FastGeoStreaming/Source/FastGeoStreaming/Public/FastGeoStreamingModule.h
@@ -17,9 +17,6 @@ public:
 	virtual void StartupModule() override;
 	virtual void ShutdownModule() override;
 
-	void AddToComponentsPendingRecreate(FFastGeoPrimitiveComponent* InComponentPendingRecreate);
-	void RemoveFromComponentsPendingRecreate(FFastGeoPrimitiveComponent* InComponentPendingRecreate);
-
 private:
 	void OnWorldPreSendAllEndOfFrameUpdates(UWorld* InWorld);
 
