[FastGeo] Fixed crash when using "r.MaterialQualityLevel" in cooked builds.
- Implemented IPrimitiveComponent & IStaticMeshComponent interfaces without creating UObjects
- Enabled IPrimitiveComponent & IStaticMeshComponent implementations to run on cooked builds

diff --git a/Engine/Plugins/Experimental/FastGeoStreaming/Source/FastGeoStreaming/Internal/FastGeoComponent.h b/Engine/Plugins/Experimental/FastGeoStreaming/Source/FastGeoStreaming/Internal/FastGeoComponent.h
--- a/Engine/Plugins/Experimental/FastGeoStreaming/Source/FastGeoStreaming/Internal/FastGeoComponent.h
+++ b/Engine/Plugins/Experimental/FastGeoStreaming/Source/FastGeoStreaming/Internal/FastGeoComponent.h
@@ -6,6 +6,7 @@
 #include "IFastGeoElement.h"
 #include "FastGeoElementType.h"
 #include "Physics/Experimental/AsyncPhysicsStateProcessorInterface.h"
+#include "Components/ComponentInterfaces.h"
 #include "FastGeoComponent.generated.h"
 
 class UWorld;
@@ -17,6 +18,12 @@ class FSceneInterface;
 class FPrimitiveSceneProxy;
 class FRegisterComponentContext;
 
+#define FASTGEO_DECLARE_COMPONENT_INTERFACE_INTERNAL(ComponentName, BasePrefix, ...)\
+	UE_DECLARE_COMPONENT_INTERFACE_INTERNAL(BasePrefix##ComponentName##__VA_ARGS__, I##ComponentName , BasePrefix##ComponentName##Interface, ComponentName##Interface, ComponentName)
+
+#define FASTGEO_DECLARE_COMPONENT_INTERFACE(ComponentName, ...)\
+	FASTGEO_DECLARE_COMPONENT_INTERFACE_INTERNAL(ComponentName, FFastGeo, __VA_ARGS__)
+
 class FASTGEOSTREAMING_API FFastGeoComponent : public IFastGeoElement
 {
 public:
@@ -30,7 +37,6 @@ public:
 
 #if WITH_EDITOR
 	virtual void InitializeFromComponent(UActorComponent* Component);
-	virtual UClass* GetEditorProxyClass() const;
 #endif
 
 	int32 GetComponentIndex() const { return ComponentIndex; }
diff --git a/Engine/Plugins/Experimental/FastGeoStreaming/Source/FastGeoStreaming/Internal/FastGeoContainer.h b/Engine/Plugins/Experimental/FastGeoStreaming/Source/FastGeoStreaming/Internal/FastGeoContainer.h
--- a/Engine/Plugins/Experimental/FastGeoStreaming/Source/FastGeoStreaming/Internal/FastGeoContainer.h
+++ b/Engine/Plugins/Experimental/FastGeoStreaming/Source/FastGeoStreaming/Internal/FastGeoContainer.h
@@ -114,6 +114,9 @@ public:
 	ULevel* GetLevel() const;
 	FFastGeoComponentCluster* GetComponentCluster(uint32 InComponentClusterTypeID, int32 InComponentClusterIndex);
 
+	TArray<void*> GetPrimitiveComponents() const;
+	TArray<void*> GetStaticMeshComponents() const;
+
 	//~ Begin UObject Interface
 	virtual void BeginDestroy() override;
 	virtual bool IsReadyForFinishDestroy() override;
@@ -187,9 +190,6 @@ protected:
 
 private:
 
-#if WITH_EDITOR
-	void CreateEditorProxyObjects();
-#endif
 	void CollectAssetReferences();
 	void SerializeComponentClusters(FArchive& Ar);
 	void RegisterToNavigationSystem();
@@ -277,17 +277,14 @@ private:
 	UPROPERTY()
 	TArray<TObjectPtr<UObject>> Assets;
 
-#if WITH_EDITORONLY_DATA
-	UPROPERTY(Transient)
-	TSet<TObjectPtr<UObject>> EditorProxyObjects;
-#endif
-
 #if !WITH_EDITOR && UE_WITH_PSO_PRECACHING
 	bool bPrecachedPSOs = false;
 	UE::Tasks::TTask<void> PrecachePSOsTask;
 #endif
 
 	// Transient data
+	TArray<void*> PrimitiveComponents;
+	TArray<void*> StaticMeshComponents;
 	TArray<FFastGeoComponent*> CollisionComponents;
 	TArray<FNavigationElementHandle> NavigationElementHandles;
 
diff --git a/Engine/Plugins/Experimental/FastGeoStreaming/Source/FastGeoStreaming/Internal/FastGeoInstancedStaticMeshComponent.h b/Engine/Plugins/Experimental/FastGeoStreaming/Source/FastGeoStreaming/Internal/FastGeoInstancedStaticMeshComponent.h
--- a/Engine/Plugins/Experimental/FastGeoStreaming/Source/FastGeoStreaming/Internal/FastGeoInstancedStaticMeshComponent.h
+++ b/Engine/Plugins/Experimental/FastGeoStreaming/Source/FastGeoStreaming/Internal/FastGeoInstancedStaticMeshComponent.h
@@ -4,7 +4,6 @@
 
 #include "FastGeoStaticMeshComponent.h"
 #include "InstancedStaticMeshSceneProxyDesc.h"
-#include "FastGeoInstancedStaticMeshComponent.generated.h"
 
 class FASTGEOSTREAMING_API FFastGeoInstancedStaticMeshComponent : public FFastGeoStaticMeshComponentBase
 {
@@ -25,7 +24,6 @@ public:
 	virtual void OnAsyncDestroyPhysicsStateEnd_GameThread() override;
 #if WITH_EDITOR
 	virtual void InitializeFromComponent(UActorComponent* Component) override;
-	virtual UClass* GetEditorProxyClass() const override;
 #endif
 	//~ End FFastGeoComponent Interface
 
@@ -91,13 +89,3 @@ private:
 
 	friend class FInstancedStaticMeshComponentHelper;
 };
-
-UCLASS()
-class UFastGeoInstancedStaticMeshComponentEditorProxy : public UFastGeoStaticMeshComponentEditorProxy
-{
-	GENERATED_BODY()
-
-#if WITH_EDITOR
-	typedef FFastGeoInstancedStaticMeshComponent ComponentType;
-#endif
-};
diff --git a/Engine/Plugins/Experimental/FastGeoStreaming/Source/FastGeoStreaming/Internal/FastGeoPrimitiveComponent.h b/Engine/Plugins/Experimental/FastGeoStreaming/Source/FastGeoStreaming/Internal/FastGeoPrimitiveComponent.h
--- a/Engine/Plugins/Experimental/FastGeoStreaming/Source/FastGeoStreaming/Internal/FastGeoPrimitiveComponent.h
+++ b/Engine/Plugins/Experimental/FastGeoStreaming/Source/FastGeoStreaming/Internal/FastGeoPrimitiveComponent.h
@@ -3,8 +3,8 @@
 #pragma once
 
 #include "FastGeoComponent.h"
+#include "FastGeoPrimitiveComponentInterface.h"
 #include "Components/PrimitiveComponent.h"
-#include "Components/ComponentInterfaces.h"
 #include "Engine/EngineTypes.h"
 #include "PrimitiveSceneInfoData.h"
 #include "PrimitiveSceneProxy.h"
@@ -14,7 +14,6 @@
 #include "Chaos/ChaosUserEntity.h"
 #include "PSOPrecacheFwd.h"
 #include "Templates/DontCopy.h"
-#include "FastGeoPrimitiveComponent.generated.h"
 
 class FPrimitiveSceneProxy;
 class UPrimitiveComponent;
@@ -67,9 +66,9 @@ private:
 
 class FASTGEOSTREAMING_API FFastGeoPrimitiveComponent : public FFastGeoComponent
 {
+	FASTGEO_DECLARE_COMPONENT_INTERFACE(PrimitiveComponent)
 public:
 	typedef FFastGeoComponent Super;
-	typedef UFastGeoPrimitiveComponentEditorProxy EditorProxyType;
 
 	/** Static type identifier for this element class */
 	static const FFastGeoElementType Type;
@@ -110,6 +109,9 @@ public:
 	};
 
 	virtual void DestroyRenderState(FFastGeoDestroyRenderStateContext* Context);
+#if WITH_EDITOR
+	virtual void NotifyRenderStateChanged();
+#endif
 	bool IsRenderStateCreated() const;
 	bool IsRenderStateDelayed() const;
 	bool IsRenderStateDirty() const;
@@ -134,7 +136,6 @@ public:
 	virtual bool IsCollisionEnabled() const override;
 #if WITH_EDITOR
 	virtual void InitializeFromComponent(UActorComponent* Component) override;
-	virtual UClass* GetEditorProxyClass() const override;
 #endif
 	//~End FFastGeoComponent interface
 
@@ -242,7 +243,6 @@ protected:
 	friend class FPrimitiveComponentHelper;
 	friend class FFastGeoPhysicsBodyInstanceOwner;
 	friend class UFastGeoContainer;
-	friend class UFastGeoPrimitiveComponentEditorProxy;
 };
 
 // Dummy type to use as base class in non editor builds
@@ -253,52 +253,3 @@ class IFastGeoDummyInterface {};
 // Would have been easier with a #ifdef in the class declaration, but UHT doesn't allow it.
 template <typename TInterface>
 class IFastGeoEditorInterface : public std::conditional<WITH_EDITOR, TInterface, IFastGeoDummyInterface<TInterface>>::type {};
-
-UCLASS()
-class UFastGeoPrimitiveComponentEditorProxy : public UFastGeoComponentEditorProxy, public IFastGeoEditorInterface<IPrimitiveComponent>
-{
-	GENERATED_BODY()
-
-#if WITH_EDITOR
-	typedef FFastGeoPrimitiveComponent ComponentType;
-
-public:
-	virtual void NotifyRenderStateChanged();
-	virtual IPrimitiveComponent* GetPrimitiveComponentInterface();
-
-protected:
-	//~ Begin UObject interface
-	virtual void BeginDestroy() override;
-	//~ End UObject interface
-
-	//~ Begin IPrimitiveComponent interface
-	virtual bool IsRenderStateCreated() const override;
-	virtual bool IsRenderStateDirty() const override;
-	virtual bool ShouldCreateRenderState() const override;
-	virtual bool IsRegistered() const override;
-	virtual bool IsUnreachable() const override;
-	virtual UWorld* GetWorld() const override;
-	virtual FSceneInterface* GetScene() const override;
-	virtual FPrimitiveSceneProxy* GetSceneProxy() const override;
-	virtual void GetStreamableRenderAssetInfo(TArray<FStreamingRenderAssetPrimitiveInfo>& ) const override {} 
-	virtual void GetUsedMaterials(TArray<UMaterialInterface*>& OutMaterials, bool bGetDebugMaterials = false) const override;
-	virtual void MarkRenderStateDirty() override;
-	virtual void DestroyRenderState() override;
-	virtual void CreateRenderState(FRegisterComponentContext* Context) override;
-	virtual FString GetName() const override;
-	virtual FString GetFullName() const override;
-	virtual FTransform GetTransform() const override;
-	virtual FBoxSphereBounds GetBounds() const override;
-	virtual float GetLastRenderTimeOnScreen() const override;
-	virtual void GetPrimitiveStats(FPrimitiveStats& PrimitiveStats) const override;
-	virtual UObject* GetUObject() override;
-	virtual const UObject* GetUObject() const override;
-	virtual void PrecachePSOs() override;
-	virtual UObject* GetOwner() const override;
-	virtual FString GetOwnerName() const override;
-	virtual FPrimitiveSceneProxy* CreateSceneProxy() override;
-	virtual HHitProxy* CreateMeshHitProxy(int32 SectionIndex, int32 MaterialIndex) override;
-	virtual HHitProxy* CreatePrimitiveHitProxies(TArray<TRefCountPtr<HHitProxy> >& OutHitProxies) override;
-	//~ End IPrimitiveComponent interface
-#endif
-};
diff --git a/Engine/Plugins/Experimental/FastGeoStreaming/Source/FastGeoStreaming/Internal/FastGeoPrimitiveComponentInterface.h b/Engine/Plugins/Experimental/FastGeoStreaming/Source/FastGeoStreaming/Internal/FastGeoPrimitiveComponentInterface.h
new file mode 100644
--- /dev/null
+++ b/Engine/Plugins/Experimental/FastGeoStreaming/Source/FastGeoStreaming/Internal/FastGeoPrimitiveComponentInterface.h
@@ -0,0 +1,49 @@
+// Copyright Epic Games, Inc. All Rights Reserved.
+
+#pragma once
+
+#include "Components/ComponentInterfaces.h"
+
+class UWorld;
+class UMaterialInterface;
+class FPrimitiveSceneProxy;
+class FSceneInterface;
+class HHitProxy;
+struct FPrimitiveStats;
+
+class FFastGeoPrimitiveComponentInterface : public IPrimitiveComponent
+{
+public:
+	bool IsRenderStateCreated() const override;
+	bool IsRenderStateDirty() const override;
+	bool ShouldCreateRenderState() const override;
+	bool IsRegistered() const override;
+	bool IsUnreachable() const override;
+	UWorld* GetWorld() const override;
+	FSceneInterface* GetScene() const override;
+	FPrimitiveSceneProxy* GetSceneProxy() const override;
+	void GetStreamableRenderAssetInfo(TArray<struct FStreamingRenderAssetPrimitiveInfo>& StreamableRenderAssets) const override;
+	void GetUsedMaterials(TArray<UMaterialInterface*>& OutMaterials, bool bGetDebugMaterials) const;
+	void MarkRenderStateDirty() override;
+	void DestroyRenderState() override;
+	void CreateRenderState(FRegisterComponentContext* Context) override;
+	void PrecachePSOs() override;
+
+	FString GetName() const override;
+	FString GetFullName() const override;
+	FTransform GetTransform() const override;
+	FBoxSphereBounds GetBounds() const override;
+	float GetLastRenderTimeOnScreen() const override;
+	void GetPrimitiveStats(FPrimitiveStats& PrimitiveStats) const override;
+
+	UObject* GetUObject() override;
+	const UObject* GetUObject() const override;
+	UObject* GetOwner() const override;
+	FString GetOwnerName() const override;
+	FPrimitiveSceneProxy* CreateSceneProxy() override;	
+
+#if WITH_EDITOR
+	HHitProxy* CreateMeshHitProxy(int32 SectionIndex, int32 MaterialIndex) override;
+#endif
+	HHitProxy* CreatePrimitiveHitProxies(TArray<TRefCountPtr<HHitProxy> >& OutHitProxies) override;
+};
diff --git a/Engine/Plugins/Experimental/FastGeoStreaming/Source/FastGeoStreaming/Internal/FastGeoStaticMeshComponent.h b/Engine/Plugins/Experimental/FastGeoStreaming/Source/FastGeoStreaming/Internal/FastGeoStaticMeshComponent.h
--- a/Engine/Plugins/Experimental/FastGeoStreaming/Source/FastGeoStreaming/Internal/FastGeoStaticMeshComponent.h
+++ b/Engine/Plugins/Experimental/FastGeoStreaming/Source/FastGeoStreaming/Internal/FastGeoStaticMeshComponent.h
@@ -4,7 +4,7 @@
 
 #include "FastGeoMeshComponent.h"
 #include "StaticMeshSceneProxyDesc.h"
-#include "FastGeoStaticMeshComponent.generated.h"
+#include "FastGeoStaticMeshComponentInterface.h"
 
 class UBodySetup;
 class UMaterialInterface;
@@ -22,6 +22,8 @@ namespace Nanite
 
 class FASTGEOSTREAMING_API FFastGeoStaticMeshComponentBase : public FFastGeoMeshComponent
 {
+	FASTGEO_DECLARE_COMPONENT_INTERFACE(StaticMeshComponent, Base)
+
 public:
 	typedef FFastGeoMeshComponent Super;
 
@@ -57,7 +59,6 @@ public:
 	virtual UBodySetup* GetBodySetup() const override;
 #if WITH_EDITOR
 	virtual void InitializeFromComponent(UActorComponent* Component) override;
-	virtual UClass* GetEditorProxyClass() const override;
 #endif
 	//~ End FFastGeoComponent interface
 
@@ -69,6 +70,9 @@ public:
 	virtual UMaterialInterface* GetMaterial(int32 MaterialIndex) const override;
 	virtual int32 GetNumMaterials() const override;
 	virtual void GetUsedMaterials(TArray<UMaterialInterface*>& OutMaterials, bool bGetDebugMaterials = false) const override;
+#if WITH_EDITOR
+	virtual void NotifyRenderStateChanged() override;
+#endif
 	//~ End FFastGeoPrimitiveComponent interface
 
 protected:
@@ -140,32 +144,4 @@ protected:
 
 private:
 	FStaticMeshSceneProxyDesc SceneProxyDesc{};
-};
-
-UCLASS()
-class UFastGeoStaticMeshComponentEditorProxy : public UFastGeoPrimitiveComponentEditorProxy, public IFastGeoEditorInterface<IStaticMeshComponent>
-{
-	GENERATED_BODY()
-
-#if WITH_EDITOR
-	typedef FFastGeoStaticMeshComponentBase ComponentType;
-
-public:
-	//~ Begin UFastGeoPrimitiveComponentEditorProxy interface
-	virtual void NotifyRenderStateChanged() override;
-	//~ End UFastGeoPrimitiveComponentEditorProxy interface
-
-protected:
-	//~ Begin IStaticMeshComponent interface
-	virtual void OnMeshRebuild(bool bRenderDataChanged) override;
-	virtual void PreStaticMeshCompilation() override;
-	virtual void PostStaticMeshCompilation() override;
-	virtual UStaticMesh* GetStaticMesh() const override;
-	
-	virtual IPrimitiveComponent* GetPrimitiveComponentInterface() override
-	{
-		return UFastGeoPrimitiveComponentEditorProxy::GetPrimitiveComponentInterface();
-	}
-	//~ End IStaticMeshComponent interface
-#endif
 };
\ No newline at end of file
diff --git a/Engine/Plugins/Experimental/FastGeoStreaming/Source/FastGeoStreaming/Internal/FastGeoStaticMeshComponentInterface.h b/Engine/Plugins/Experimental/FastGeoStreaming/Source/FastGeoStreaming/Internal/FastGeoStaticMeshComponentInterface.h
new file mode 100644
--- /dev/null
+++ b/Engine/Plugins/Experimental/FastGeoStreaming/Source/FastGeoStreaming/Internal/FastGeoStaticMeshComponentInterface.h
@@ -0,0 +1,19 @@
+// Copyright Epic Games, Inc. All Rights Reserved.
+
+#pragma once 
+
+#include "Components/ComponentInterfaces.h"
+
+class UStaticMesh;
+
+class FFastGeoStaticMeshComponentInterface : public IStaticMeshComponent
+{
+public:
+#if WITH_EDITOR
+	virtual void OnMeshRebuild(bool bRenderDataChanged) override {}
+	virtual void PreStaticMeshCompilation() override {}
+	virtual void PostStaticMeshCompilation() override {}
+#endif
+	virtual UStaticMesh* GetStaticMesh() const override;
+	virtual IPrimitiveComponent* GetPrimitiveComponentInterface() override;
+};
diff --git a/Engine/Plugins/Experimental/FastGeoStreaming/Source/FastGeoStreaming/Private/FastGeoComponent.cpp b/Engine/Plugins/Experimental/FastGeoStreaming/Source/FastGeoStreaming/Private/FastGeoComponent.cpp
--- a/Engine/Plugins/Experimental/FastGeoStreaming/Source/FastGeoStreaming/Private/FastGeoComponent.cpp
+++ b/Engine/Plugins/Experimental/FastGeoStreaming/Source/FastGeoStreaming/Private/FastGeoComponent.cpp
@@ -21,11 +21,6 @@ FFastGeoComponent::FFastGeoComponent(int32 InComponentIndex, FFastGeoElementType
 void FFastGeoComponent::InitializeFromComponent(UActorComponent* InComponent)
 {
 }
-
-UClass* FFastGeoComponent::GetEditorProxyClass() const
-{
-	return UFastGeoComponentEditorProxy::StaticClass();
-}
 #endif
 
 FArchive& operator<<(FArchive& Ar, FFastGeoComponent& Component)
diff --git a/Engine/Plugins/Experimental/FastGeoStreaming/Source/FastGeoStreaming/Private/FastGeoContainer.cpp b/Engine/Plugins/Experimental/FastGeoStreaming/Source/FastGeoStreaming/Private/FastGeoContainer.cpp
--- a/Engine/Plugins/Experimental/FastGeoStreaming/Source/FastGeoStreaming/Private/FastGeoContainer.cpp
+++ b/Engine/Plugins/Experimental/FastGeoStreaming/Source/FastGeoStreaming/Private/FastGeoContainer.cpp
@@ -313,6 +313,16 @@ FFastGeoComponentCluster* UFastGeoContainer::GetComponentCluster(uint32 InCompon
 	return nullptr;
 }
 
+TArray<void*> UFastGeoContainer::GetPrimitiveComponents() const
+{
+	return PrimitiveComponents;
+}
+
+TArray<void*> UFastGeoContainer::GetStaticMeshComponents() const
+{
+	return StaticMeshComponents;
+}
+
 ULevel* UFastGeoContainer::GetLevel() const
 {
 	return GetOuterULevel();
@@ -516,6 +526,19 @@ void UFastGeoContainer::BeginDestroy()
 {
 	Super::BeginDestroy();
 	DestroyFence.BeginFence();
+
+#if WITH_EDITOR
+	if (!IsTemplate())
+	{
+		ForEachComponentCluster([](FFastGeoComponentCluster& ComponentCluster)
+		{
+			ComponentCluster.ForEachComponent<FFastGeoPrimitiveComponent>([](FFastGeoPrimitiveComponent& Component)
+			{
+				Component.NotifyRenderStateChanged();
+			});
+		});
+	}
+#endif
 }
 
 bool UFastGeoContainer::IsReadyForFinishDestroy()
@@ -543,23 +566,6 @@ void UFastGeoContainer::AddComponentCluster(FFastGeoComponentCluster* ComponentC
 }
 
 #if WITH_EDITOR
-void UFastGeoContainer::CreateEditorProxyObjects()
-{
-	ForEachComponentCluster([this](FFastGeoComponentCluster& ComponentCluster)
-	{
-		ComponentCluster.ForEachComponent([this](FFastGeoComponent& Component)
-		{
-			if (UClass* EditorProxyClass = Component.GetEditorProxyClass())
-			{
-				UFastGeoComponentEditorProxy* ComponentEditorProxy = NewObject<UFastGeoComponentEditorProxy>(this, EditorProxyClass);
-				ComponentEditorProxy->SetFastGeoComponent(&Component);
-				Component.SetEditorProxy(ComponentEditorProxy);
-				EditorProxyObjects.Add(ComponentEditorProxy);
-			}
-		});
-	});
-}
-
 void UFastGeoContainer::PreSave(FObjectPreSaveContext ObjectSaveContext)
 {
 	Super::PreSave(ObjectSaveContext);
@@ -642,11 +648,6 @@ void UFastGeoContainer::OnCreated(bool bCollectReferences)
 	// Initialize component clusters & components dynamic properties
 	InitializeDynamicProperties(bIsGameWorld);
 
-#if WITH_EDITOR
-	// In editor, we need a UObject representation of the components for some operations
-	CreateEditorProxyObjects();
-#endif
-
 	// Always collect references outside of game worlds.
 	if (bCollectReferences || !bIsGameWorld)
 	{
@@ -725,6 +726,8 @@ void UFastGeoContainer::InitializeDynamicProperties(bool bInitForPlay)
 
 	if (bInitForPlay)
 	{
+		PrimitiveComponents.Reset();
+		StaticMeshComponents.Reset();
 		CollisionComponents.Reset();
 		ForEachComponentCluster([this](FFastGeoComponentCluster& ComponentCluster)
 		{
@@ -734,6 +737,15 @@ void UFastGeoContainer::InitializeDynamicProperties(bool bInitForPlay)
 				{
 					CollisionComponents.Add(&Component);
 				}
+
+				if (FFastGeoPrimitiveComponent* PrimitiveComponent = Component.CastTo<FFastGeoPrimitiveComponent>())
+				{
+					PrimitiveComponents.Add(PrimitiveComponent->GetPrimitiveComponentInterface());
+				}
+				if (FFastGeoStaticMeshComponent* StaticMeshComponent = Component.CastTo<FFastGeoStaticMeshComponent>())
+				{
+					StaticMeshComponents.Add(StaticMeshComponent->GetStaticMeshComponentInterface());
+				}
 			});
 		});
 	}
diff --git a/Engine/Plugins/Experimental/FastGeoStreaming/Source/FastGeoStreaming/Private/FastGeoInstancedStaticMeshComponent.cpp b/Engine/Plugins/Experimental/FastGeoStreaming/Source/FastGeoStreaming/Private/FastGeoInstancedStaticMeshComponent.cpp
--- a/Engine/Plugins/Experimental/FastGeoStreaming/Source/FastGeoStreaming/Private/FastGeoInstancedStaticMeshComponent.cpp
+++ b/Engine/Plugins/Experimental/FastGeoStreaming/Source/FastGeoStreaming/Private/FastGeoInstancedStaticMeshComponent.cpp
@@ -69,11 +69,6 @@ void FFastGeoInstancedStaticMeshComponent::InitializeFromComponent(UActorCompone
 	NavigationBounds = CalculateBounds(EBoundsType::NavigationBounds).GetBox();
 }
 
-UClass* FFastGeoInstancedStaticMeshComponent::GetEditorProxyClass() const
-{
-	return UFastGeoInstancedStaticMeshComponentEditorProxy::StaticClass();
-}
-
 void FFastGeoInstancedStaticMeshComponent::ResetSceneProxyDescUnsupportedProperties()
 {
 	Super::ResetSceneProxyDescUnsupportedProperties();
diff --git a/Engine/Plugins/Experimental/FastGeoStreaming/Source/FastGeoStreaming/Private/FastGeoPrimitiveComponent.cpp b/Engine/Plugins/Experimental/FastGeoStreaming/Source/FastGeoStreaming/Private/FastGeoPrimitiveComponent.cpp
--- a/Engine/Plugins/Experimental/FastGeoStreaming/Source/FastGeoStreaming/Private/FastGeoPrimitiveComponent.cpp
+++ b/Engine/Plugins/Experimental/FastGeoStreaming/Source/FastGeoStreaming/Private/FastGeoPrimitiveComponent.cpp
@@ -23,6 +23,10 @@
 #include "PrimitiveComponentHelper.h"
 #include "UnrealEngine.h"
 
+#if WITH_EDITOR
+#include "ObjectCacheEventSink.h"
+#endif
+
 #if !UE_BUILD_SHIPPING
 namespace FastGeo
 {
@@ -227,11 +231,6 @@ void FFastGeoPrimitiveComponent::InitializeFromComponent(UActorComponent* Compon
 	ResetSceneProxyDescUnsupportedProperties();
 }
 
-UClass* FFastGeoPrimitiveComponent::GetEditorProxyClass() const
-{
-	return UFastGeoPrimitiveComponentEditorProxy::StaticClass();
-}
-
 void FFastGeoPrimitiveComponent::ResetSceneProxyDescUnsupportedProperties()
 {
 	// Unsupported properties
@@ -312,7 +311,7 @@ void FFastGeoPrimitiveComponent::InitializeSceneProxyDescDynamicProperties()
 	SceneProxyDesc.CustomPrimitiveData = &CustomPrimitiveData;
 	SceneProxyDesc.Scene = GetScene();
 #if WITH_EDITOR
-	SceneProxyDesc.PrimitiveComponentInterface = GetEditorProxy<UFastGeoPrimitiveComponentEditorProxy>()->GetPrimitiveComponentInterface();
+	SceneProxyDesc.PrimitiveComponentInterface = GetPrimitiveComponentInterface();
 #endif
 	SceneProxyDesc.FeatureLevel = SceneProxyDesc.Scene->GetFeatureLevel();
 	TArray<URuntimeVirtualTexture*> const& VirtualTextures = GetRuntimeVirtualTextures();
@@ -352,7 +351,7 @@ void FFastGeoPrimitiveComponent::CreateRenderState(FRegisterComponentContext* Co
 	bRenderStateDirty = false;
 
 #if WITH_EDITOR
-	GetEditorProxy<EditorProxyType>()->NotifyRenderStateChanged();
+	NotifyRenderStateChanged();
 #endif
 	
 	FSceneInterface* Scene = GetScene();
@@ -424,7 +423,7 @@ void FFastGeoPrimitiveComponent::DestroyRenderState(FFastGeoDestroyRenderStateCo
 		bRenderStateDirty = false;
 
 #if WITH_EDITOR
-		GetEditorProxy<EditorProxyType>()->NotifyRenderStateChanged();
+		NotifyRenderStateChanged();
 #endif
 	}
 	else
@@ -434,6 +433,13 @@ void FFastGeoPrimitiveComponent::DestroyRenderState(FFastGeoDestroyRenderStateCo
 	}
 }
 
+#if WITH_EDITOR
+void FFastGeoPrimitiveComponent::NotifyRenderStateChanged()
+{
+	FObjectCacheEventSink::NotifyRenderStateChanged_Concurrent(GetPrimitiveComponentInterface());
+}
+#endif
+
 void FFastGeoPrimitiveComponent::OnAsyncCreatePhysicsState()
 {
 	TRACE_CPUPROFILER_EVENT_SCOPE(FFastGeoPrimitiveComponent::OnAsyncCreatePhysicsState);
@@ -586,6 +592,7 @@ FPrimitiveSceneProxy* FFastGeoPrimitiveComponent::GetSceneProxy() const
 
 void FFastGeoPrimitiveComponent::MarkRenderStateDirty()
 {
+	FReadScopeLock ReadLock(*Lock.Get());
 	if ((IsRenderStateCreated() || IsRenderStateDelayed()) && !IsRenderStateDirty())
 	{
 		bRenderStateDirty = true;
@@ -773,7 +780,11 @@ void FFastGeoPrimitiveComponent::PrecachePSOs()
 {
 #if !WITH_EDITOR && UE_WITH_PSO_PRECACHING
 	TRACE_CPUPROFILER_EVENT_SCOPE(FFastGeoPrimitiveComponent::PrecachePSOs);
-	check(bPSOPrecacheRequired);
+
+	if (!bPSOPrecacheRequired)
+	{
+		return;
+	}
 
 	// Clear the current request data
 	MaterialPSOPrecacheRequestIDs.Empty();
@@ -940,156 +951,3 @@ void FFastGeoPhysicsBodyInstanceOwner::GetComplexPhysicalMaterials(TArray<UPhysi
 	check(OwnerContainer.IsValid());
 	return OwnerComponent->GetComplexPhysicalMaterials(OutPhysMaterials, OutPhysMaterialMasks);
 }
-
-#if WITH_EDITOR
-void UFastGeoPrimitiveComponentEditorProxy::NotifyRenderStateChanged()
-{
-	FObjectCacheEventSink::NotifyRenderStateChanged_Concurrent(this);
-}
-
-IPrimitiveComponent* UFastGeoPrimitiveComponentEditorProxy::GetPrimitiveComponentInterface()
-{
-	return this;
-}
-
-void UFastGeoPrimitiveComponentEditorProxy::BeginDestroy()
-{
-	Super::BeginDestroy();
-
-	if (!IsTemplate())
-	{
-		NotifyRenderStateChanged();
-	}
-}
-
-//~ Begin IPrimitiveComponent interface
-bool UFastGeoPrimitiveComponentEditorProxy::IsRenderStateCreated() const
-{
-	return GetComponent<ComponentType>().IsRenderStateCreated();
-}
-
-bool UFastGeoPrimitiveComponentEditorProxy::IsRenderStateDirty() const
-{
-	return GetComponent<ComponentType>().IsRenderStateDirty();
-}
-
-bool UFastGeoPrimitiveComponentEditorProxy::ShouldCreateRenderState() const
-{
-	return GetComponent<ComponentType>().ShouldCreateRenderState();
-}
-
-bool UFastGeoPrimitiveComponentEditorProxy::IsRegistered() const
-{
-	return GetComponent<ComponentType>().IsRegistered();
-}
-
-bool UFastGeoPrimitiveComponentEditorProxy::IsUnreachable() const
-{
-	return Super::IsUnreachable();
-}
-
-UWorld* UFastGeoPrimitiveComponentEditorProxy::GetWorld() const
-{
-	return GetComponent<ComponentType>().GetWorld();
-}
-
-FSceneInterface* UFastGeoPrimitiveComponentEditorProxy::GetScene() const
-{
-	return GetComponent<ComponentType>().GetScene();
-}
-
-FPrimitiveSceneProxy* UFastGeoPrimitiveComponentEditorProxy::GetSceneProxy() const
-{
-	return GetComponent<ComponentType>().GetSceneProxy();
-}
-
-void UFastGeoPrimitiveComponentEditorProxy::GetUsedMaterials(TArray<UMaterialInterface*>& OutMaterials, bool bGetDebugMaterials) const
-{
-	GetComponent<ComponentType>().GetUsedMaterials(OutMaterials, bGetDebugMaterials);
-}
-
-void UFastGeoPrimitiveComponentEditorProxy::MarkRenderStateDirty()
-{
-	GetComponent<ComponentType>().MarkRenderStateDirty();
-}
-
-void UFastGeoPrimitiveComponentEditorProxy::DestroyRenderState()
-{
-	GetComponent<ComponentType>().DestroyRenderState(nullptr);
-}
-
-void UFastGeoPrimitiveComponentEditorProxy::CreateRenderState(FRegisterComponentContext* Context)
-{
-	GetComponent<ComponentType>().CreateRenderState(Context);
-}
-
-FString UFastGeoPrimitiveComponentEditorProxy::GetName() const
-{
-	return GetUObject()->GetName();
-}
-
-FString UFastGeoPrimitiveComponentEditorProxy::GetFullName() const
-{
-	return GetUObject()->GetFullName();
-}
-
-FTransform UFastGeoPrimitiveComponentEditorProxy::GetTransform() const
-{
-	return GetComponent<ComponentType>().GetTransform();
-}
-
-FBoxSphereBounds UFastGeoPrimitiveComponentEditorProxy::GetBounds() const
-{
-	return GetComponent<ComponentType>().GetBounds();
-}
-
-float UFastGeoPrimitiveComponentEditorProxy::GetLastRenderTimeOnScreen() const
-{
-	return GetComponent<ComponentType>().GetLastRenderTimeOnScreen();
-}
-
-void UFastGeoPrimitiveComponentEditorProxy::GetPrimitiveStats(FPrimitiveStats& PrimitiveStats) const
-{
-}
-
-UObject* UFastGeoPrimitiveComponentEditorProxy::GetUObject()
-{
-	return this;
-}
-
-const UObject* UFastGeoPrimitiveComponentEditorProxy::GetUObject() const
-{
-	return this;
-}
-
-void UFastGeoPrimitiveComponentEditorProxy::PrecachePSOs()
-{
-	GetComponent<ComponentType>().PrecachePSOs();
-}
-
-UObject* UFastGeoPrimitiveComponentEditorProxy::GetOwner() const
-{
-	return GetOuter();
-}
-
-FString UFastGeoPrimitiveComponentEditorProxy::GetOwnerName() const
-{
-	return GetOwner()->GetName();
-}
-
-FPrimitiveSceneProxy* UFastGeoPrimitiveComponentEditorProxy::CreateSceneProxy()
-{
-	return GetComponent<ComponentType>().CreateSceneProxy();
-}
-
-HHitProxy* UFastGeoPrimitiveComponentEditorProxy::CreateMeshHitProxy(int32 SectionIndex, int32 MaterialIndex)
-{
-	return nullptr;
-}
-
-HHitProxy* UFastGeoPrimitiveComponentEditorProxy::CreatePrimitiveHitProxies(TArray<TRefCountPtr<HHitProxy>>& OutHitProxies)
-{
-	return nullptr;
-}
-//~ End IPrimitiveComponent interface
-#endif
diff --git a/Engine/Plugins/Experimental/FastGeoStreaming/Source/FastGeoStreaming/Private/FastGeoPrimitiveComponentInterface.cpp b/Engine/Plugins/Experimental/FastGeoStreaming/Source/FastGeoStreaming/Private/FastGeoPrimitiveComponentInterface.cpp
new file mode 100644
--- /dev/null
+++ b/Engine/Plugins/Experimental/FastGeoStreaming/Source/FastGeoStreaming/Private/FastGeoPrimitiveComponentInterface.cpp
@@ -0,0 +1,142 @@
+// Copyright Epic Games, Inc. All Rights Reserved.
+
+#include "FastGeoPrimitiveComponentInterface.h"
+#include "FastGeoPrimitiveComponent.h"
+#include "FastGeoContainer.h"
+#include "Materials/MaterialInterface.h"
+
+bool FFastGeoPrimitiveComponentInterface::IsRenderStateCreated() const
+{
+	return FFastGeoPrimitiveComponent::GetPrimitiveComponent(this)->IsRenderStateCreated();
+}
+
+bool FFastGeoPrimitiveComponentInterface::IsRenderStateDirty() const
+{
+	return FFastGeoPrimitiveComponent::GetPrimitiveComponent(this)->IsRenderStateDirty();
+}
+
+bool FFastGeoPrimitiveComponentInterface::ShouldCreateRenderState() const
+{
+	return FFastGeoPrimitiveComponent::GetPrimitiveComponent(this)->ShouldCreateRenderState();
+}
+
+bool FFastGeoPrimitiveComponentInterface::IsRegistered() const
+{
+	return FFastGeoPrimitiveComponent::GetPrimitiveComponent(this)->IsRegistered();
+}
+
+bool FFastGeoPrimitiveComponentInterface::IsUnreachable() const
+{
+	const UObject* Owner = GetOwner();
+	return !Owner || Owner->IsUnreachable();
+}
+
+UWorld* FFastGeoPrimitiveComponentInterface::GetWorld() const
+{
+	return FFastGeoPrimitiveComponent::GetPrimitiveComponent(this)->GetWorld();
+}
+
+FSceneInterface* FFastGeoPrimitiveComponentInterface::GetScene() const
+{
+	return FFastGeoPrimitiveComponent::GetPrimitiveComponent(this)->GetScene();
+}
+
+FPrimitiveSceneProxy* FFastGeoPrimitiveComponentInterface::GetSceneProxy() const
+{
+	return FFastGeoPrimitiveComponent::GetPrimitiveComponent(this)->GetSceneProxy();
+}
+
+void FFastGeoPrimitiveComponentInterface::GetUsedMaterials(TArray<UMaterialInterface*>& OutMaterials, bool bGetDebugMaterials) const
+{
+	FFastGeoPrimitiveComponent::GetPrimitiveComponent(this)->GetUsedMaterials(OutMaterials, bGetDebugMaterials);
+}
+
+void FFastGeoPrimitiveComponentInterface::GetStreamableRenderAssetInfo(TArray<FStreamingRenderAssetPrimitiveInfo>& StreamableRenderAssets) const
+{
+}
+
+void FFastGeoPrimitiveComponentInterface::MarkRenderStateDirty()
+{
+	FFastGeoPrimitiveComponent::GetPrimitiveComponent(this)->MarkRenderStateDirty();
+}
+
+void FFastGeoPrimitiveComponentInterface::DestroyRenderState()
+{
+	FFastGeoPrimitiveComponent::GetPrimitiveComponent(this)->DestroyRenderState(nullptr);
+}
+
+void FFastGeoPrimitiveComponentInterface::CreateRenderState(FRegisterComponentContext* Context)
+{
+	FFastGeoPrimitiveComponent::GetPrimitiveComponent(this)->CreateRenderState(Context);
+}
+
+FString FFastGeoPrimitiveComponentInterface::GetName() const
+{
+	return GetNameSafe(GetUObject());
+}
+
+FString FFastGeoPrimitiveComponentInterface::GetFullName() const
+{
+	return GetFullNameSafe(GetUObject());
+}
+
+FTransform FFastGeoPrimitiveComponentInterface::GetTransform() const
+{
+	return FFastGeoPrimitiveComponent::GetPrimitiveComponent(this)->GetComponentTransform();
+}
+
+FBoxSphereBounds FFastGeoPrimitiveComponentInterface::GetBounds() const
+{
+	return FFastGeoPrimitiveComponent::GetPrimitiveComponent(this)->GetBounds();
+}
+
+float FFastGeoPrimitiveComponentInterface::GetLastRenderTimeOnScreen() const
+{
+	return FFastGeoPrimitiveComponent::GetPrimitiveComponent(this)->GetLastRenderTimeOnScreen();
+}
+
+void FFastGeoPrimitiveComponentInterface::GetPrimitiveStats(FPrimitiveStats& PrimitiveStats) const
+{
+}
+
+UObject* FFastGeoPrimitiveComponentInterface::GetUObject()
+{
+	return GetOwner();
+}
+
+const UObject* FFastGeoPrimitiveComponentInterface::GetUObject() const
+{
+	return GetOwner();
+}
+
+UObject* FFastGeoPrimitiveComponentInterface::GetOwner() const
+{
+	return FFastGeoPrimitiveComponent::GetPrimitiveComponent(this)->GetOwnerContainer();
+}
+
+FString FFastGeoPrimitiveComponentInterface::GetOwnerName() const
+{
+	return GetNameSafe(GetOwner());
+}
+
+FPrimitiveSceneProxy* FFastGeoPrimitiveComponentInterface::CreateSceneProxy()
+{
+	return FFastGeoPrimitiveComponent::GetPrimitiveComponent(this)->CreateSceneProxy();
+}
+
+void FFastGeoPrimitiveComponentInterface::PrecachePSOs()
+{
+	FFastGeoPrimitiveComponent::GetPrimitiveComponent(this)->PrecachePSOs();
+}
+
+HHitProxy* FFastGeoPrimitiveComponentInterface::CreatePrimitiveHitProxies(TArray<TRefCountPtr<HHitProxy> >& OutHitProxies)
+{
+	return nullptr;
+}
+
+#if WITH_EDITOR
+HHitProxy* FFastGeoPrimitiveComponentInterface::CreateMeshHitProxy(int32 SectionIndex, int32 MaterialIndex)
+{
+	return nullptr;
+}
+#endif
\ No newline at end of file
diff --git a/Engine/Plugins/Experimental/FastGeoStreaming/Source/FastGeoStreaming/Private/FastGeoStaticMeshComponent.cpp b/Engine/Plugins/Experimental/FastGeoStreaming/Source/FastGeoStreaming/Private/FastGeoStaticMeshComponent.cpp
--- a/Engine/Plugins/Experimental/FastGeoStreaming/Source/FastGeoStreaming/Private/FastGeoStaticMeshComponent.cpp
+++ b/Engine/Plugins/Experimental/FastGeoStreaming/Source/FastGeoStreaming/Private/FastGeoStaticMeshComponent.cpp
@@ -88,6 +88,15 @@ UBodySetup* FFastGeoStaticMeshComponentBase::GetBodySetup() const
 	return GetStaticMesh() ? GetStaticMesh()->GetBodySetup() : nullptr;
 }
 
+#if WITH_EDITOR
+void FFastGeoStaticMeshComponentBase::NotifyRenderStateChanged()
+{
+	Super::NotifyRenderStateChanged();
+
+	FObjectCacheEventSink::NotifyStaticMeshChanged_Concurrent(GetStaticMeshComponentInterface());
+}
+#endif
+
 bool FFastGeoStaticMeshComponentBase::IsNavigationRelevant() const
 {
 	return FStaticMeshComponentHelper::IsNavigationRelevant(*this);
@@ -171,11 +180,6 @@ void FFastGeoStaticMeshComponentBase::InitializeFromComponent(UActorComponent* C
 	WorldBounds = LocalBounds.TransformBy(WorldTransform);
 }
 
-UClass* FFastGeoStaticMeshComponentBase::GetEditorProxyClass() const
-{
-	return UFastGeoStaticMeshComponentEditorProxy::StaticClass();
-}
-
 void FFastGeoStaticMeshComponentBase::ResetSceneProxyDescUnsupportedProperties()
 {
 	Super::ResetSceneProxyDescUnsupportedProperties();
@@ -450,36 +454,6 @@ void FFastGeoStaticMeshComponentBase::CollectPSOPrecacheData(const FPSOPrecacheP
 }
 #endif
 
-#if WITH_EDITOR
-
-void UFastGeoStaticMeshComponentEditorProxy::NotifyRenderStateChanged()
-{
-	Super::NotifyRenderStateChanged();
-
-	FObjectCacheEventSink::NotifyStaticMeshChanged_Concurrent(this);
-}
-
-//~ Begin IStaticMeshComponent interface
-void UFastGeoStaticMeshComponentEditorProxy::OnMeshRebuild(bool bRenderDataChanged)
-{
-}
-
-void UFastGeoStaticMeshComponentEditorProxy::PreStaticMeshCompilation()
-{
-}
-
-void UFastGeoStaticMeshComponentEditorProxy::PostStaticMeshCompilation()
-{
-}
-
-UStaticMesh* UFastGeoStaticMeshComponentEditorProxy::GetStaticMesh() const
-{
-	return GetComponent<ComponentType>().GetStaticMesh();
-}
-//~ End IStaticMeshComponent interface
-
-#endif
-
 const FFastGeoElementType FFastGeoStaticMeshComponent::Type(&FFastGeoStaticMeshComponentBase::Type);
 
 FFastGeoStaticMeshComponent::FFastGeoStaticMeshComponent(int32 InComponentIndex, FFastGeoElementType InType)
diff --git a/Engine/Plugins/Experimental/FastGeoStreaming/Source/FastGeoStreaming/Private/FastGeoStaticMeshComponentInterface.cpp b/Engine/Plugins/Experimental/FastGeoStreaming/Source/FastGeoStreaming/Private/FastGeoStaticMeshComponentInterface.cpp
new file mode 100644
--- /dev/null
+++ b/Engine/Plugins/Experimental/FastGeoStreaming/Source/FastGeoStreaming/Private/FastGeoStaticMeshComponentInterface.cpp
@@ -0,0 +1,14 @@
+// Copyright Epic Games, Inc. All Rights Reserved.
+
+#include "FastGeoStaticMeshComponentInterface.h"
+#include "FastGeoStaticMeshComponent.h"
+
+UStaticMesh* FFastGeoStaticMeshComponentInterface::GetStaticMesh() const
+{
+	return FFastGeoStaticMeshComponentBase::GetStaticMeshComponent(this)->GetStaticMesh();
+}
+
+IPrimitiveComponent* FFastGeoStaticMeshComponentInterface::GetPrimitiveComponentInterface()
+{
+	return FFastGeoStaticMeshComponentBase::GetStaticMeshComponent(this)->GetPrimitiveComponentInterface();
+}
diff --git a/Engine/Plugins/Experimental/FastGeoStreaming/Source/FastGeoStreaming/Private/FastGeoStreamingModule.cpp b/Engine/Plugins/Experimental/FastGeoStreaming/Source/FastGeoStreaming/Private/FastGeoStreamingModule.cpp
--- a/Engine/Plugins/Experimental/FastGeoStreaming/Source/FastGeoStreaming/Private/FastGeoStreamingModule.cpp
+++ b/Engine/Plugins/Experimental/FastGeoStreaming/Source/FastGeoStreaming/Private/FastGeoStreamingModule.cpp
@@ -12,39 +12,23 @@ IMPLEMENT_MODULE(FFastGeoStreamingModule, FastGeoStreaming);
 
 void FFastGeoStreamingModule::StartupModule()
 {
-#if WITH_EDITOR
-	IPrimitiveComponent::AddImplementer({ UFastGeoStaticMeshComponentEditorProxy::StaticClass(),  [](UObject* Obj)
+	IPrimitiveComponent::AddImplementer(FComponentInterfaceImplementation(UFastGeoContainer::StaticClass(),  [](UObject* Obj)
 	{
-		return static_cast<IPrimitiveComponent*>(static_cast<UFastGeoPrimitiveComponentEditorProxy*>(Obj));
-	}});
+		return CastChecked<UFastGeoContainer>(Obj)->GetPrimitiveComponents();
+	}));
 
-	IStaticMeshComponent::AddImplementer({ UFastGeoStaticMeshComponentEditorProxy::StaticClass(), [](UObject* Obj)
+	IStaticMeshComponent::AddImplementer(FComponentInterfaceImplementation(UFastGeoContainer::StaticClass(),  [](UObject* Obj)
 	{
-		return static_cast<IStaticMeshComponent*>(static_cast<UFastGeoStaticMeshComponentEditorProxy*>(Obj));
-	}});
-
-	IPrimitiveComponent::AddImplementer({ UFastGeoInstancedStaticMeshComponentEditorProxy::StaticClass(), [](UObject* Obj)
-	{
-		return static_cast<IPrimitiveComponent*>(static_cast<UFastGeoPrimitiveComponentEditorProxy*>(Obj));
-	}});
-
-	IStaticMeshComponent::AddImplementer({ UFastGeoInstancedStaticMeshComponentEditorProxy::StaticClass(), [](UObject* Obj)
-	{
-		return static_cast<IStaticMeshComponent*>(static_cast<UFastGeoStaticMeshComponentEditorProxy*>(Obj));
-	}});
-#endif
+		return CastChecked<UFastGeoContainer>(Obj)->GetStaticMeshComponents();
+	}));
 
 	Handle_OnWorldPreSendAllEndOfFrameUpdates = FWorldDelegates::OnWorldPreSendAllEndOfFrameUpdates.AddRaw(this, &FFastGeoStreamingModule::OnWorldPreSendAllEndOfFrameUpdates);
 }
 
 void FFastGeoStreamingModule::ShutdownModule()
 {
-#if WITH_EDITOR
-	IPrimitiveComponent::RemoveImplementer(UFastGeoStaticMeshComponentEditorProxy::StaticClass());
-	IStaticMeshComponent::RemoveImplementer(UFastGeoStaticMeshComponentEditorProxy::StaticClass());
-	IPrimitiveComponent::RemoveImplementer(UFastGeoInstancedStaticMeshComponentEditorProxy::StaticClass());
-	IStaticMeshComponent::RemoveImplementer(UFastGeoInstancedStaticMeshComponentEditorProxy::StaticClass());
-#endif
+	IPrimitiveComponent::RemoveImplementer(UFastGeoContainer::StaticClass());
+	IStaticMeshComponent::RemoveImplementer(UFastGeoContainer::StaticClass());
 
 	FWorldDelegates::OnWorldPreSendAllEndOfFrameUpdates.Remove(Handle_OnWorldPreSendAllEndOfFrameUpdates);
 }
diff --git a/Engine/Source/Runtime/Engine/Classes/Components/ComponentInterfaceIterator.h b/Engine/Source/Runtime/Engine/Classes/Components/ComponentInterfaceIterator.h
--- a/Engine/Source/Runtime/Engine/Classes/Components/ComponentInterfaceIterator.h
+++ b/Engine/Source/Runtime/Engine/Classes/Components/ComponentInterfaceIterator.h
@@ -22,14 +22,37 @@ public:
 
 		for (FComponentInterfaceImplementation& Implementation : T::Implementers)
 		{
-			ObjectArray.Reset();
-			GetObjectsOfClass(Implementation.Class, ObjectArray, bIncludeDerivedClasses, AdditionalExclusionFlags, GetObjectIteratorDefaultInternalExclusionFlags(InInternalExclusionFlags));
-			Interfaces.Reserve(ObjectArray.Num() + Interfaces.Num());
-			for (UObject* Object : ObjectArray)
+			if (Implementation.Class)
 			{
-				T* Interface = reinterpret_cast<T*>(Implementation.Resolver(Object));
-				check(Interface);
-				Interfaces.Add(Interface);
+				if (Implementation.Resolver)
+				{
+					ObjectArray.Reset();
+					GetObjectsOfClass(Implementation.Class, ObjectArray, bIncludeDerivedClasses, AdditionalExclusionFlags, GetObjectIteratorDefaultInternalExclusionFlags(InInternalExclusionFlags));
+					Interfaces.Reserve(ObjectArray.Num() + Interfaces.Num());
+					for (UObject* Object : ObjectArray)
+					{
+						T* Interface = reinterpret_cast<T*>(Implementation.Resolver(Object));
+						check(Interface);
+						Interfaces.Add(Interface);
+					}
+				}
+
+				if (Implementation.CollectionResolver)
+				{
+					ObjectArray.Reset();
+					GetObjectsOfClass(Implementation.Class, ObjectArray, bIncludeDerivedClasses, AdditionalExclusionFlags, GetObjectIteratorDefaultInternalExclusionFlags(InInternalExclusionFlags));
+					for (UObject* Object : ObjectArray)
+					{
+						TArray<void*> ContainerInterfaces = Implementation.CollectionResolver(Object);
+						Interfaces.Reserve(ContainerInterfaces.Num() + Interfaces.Num());
+						for (void* ContainerInterface : ContainerInterfaces)
+						{
+							T* Interface = reinterpret_cast<T*>(ContainerInterface);
+							check(Interface);
+							Interfaces.Add(Interface);
+						}
+					}
+				}
 			}
 		}
 
diff --git a/Engine/Source/Runtime/Engine/Classes/Components/ComponentInterfaces.h b/Engine/Source/Runtime/Engine/Classes/Components/ComponentInterfaces.h
--- a/Engine/Source/Runtime/Engine/Classes/Components/ComponentInterfaces.h
+++ b/Engine/Source/Runtime/Engine/Classes/Components/ComponentInterfaces.h
@@ -52,8 +52,19 @@ struct FPrimitiveStats
 
 struct FComponentInterfaceImplementation
 {
-	UClass*	Class;
+	FComponentInterfaceImplementation(UClass* InClass, TFunction<void* (UObject*)> InResolver)
+		: Class(InClass)
+		, Resolver(InResolver)
+	{}
+
+	FComponentInterfaceImplementation(UClass* InClass, TFunction<TArray<void*>(UObject*)> InCollectionResolver)
+		: Class(InClass)
+		, CollectionResolver(InCollectionResolver)
+	{}
+
+	UClass*	Class = nullptr;
 	TFunction<void*(UObject*)> Resolver;
+	TFunction<TArray<void*> (UObject*)> CollectionResolver;
 };
 
 class IPrimitiveComponent
