[FastGeo] Renamed FastGeoStreaming.ini to DefaultFastGeoStreaming.ini following 46386808

diff --git a/Engine/Plugins/Experimental/FastGeoStreaming/Config/DefaultFastGeoStreaming.ini b/Engine/Plugins/Experimental/FastGeoStreaming/Config/DefaultFastGeoStreaming.ini
new file mode 100644
--- /dev/null
+++ b/Engine/Plugins/Experimental/FastGeoStreaming/Config/DefaultFastGeoStreaming.ini
@@ -0,0 +1,15 @@
+[/Script/FastGeoStreaming.FastGeoWorldPartitionRuntimeCellTransformer]
++BuiltinAllowedActorClasses=/Script/CoreUObject.Class'/Script/Engine.PackedLevelActor'
++BuiltinAllowedActorClasses=/Script/CoreUObject.Class'/Script/Engine.StaticMeshActor'
++BuiltinAllowedActorClasses=/Script/CoreUObject.Class'/Script/Engine.WorldPartitionHLOD'
++BuiltinAllowedActorClasses=/Script/CoreUObject.Class'/Script/Engine.WorldPartitionAutoInstancedActor'
++BuiltinDisallowedActorClasses=/Script/CoreUObject.Class'/Script/Water.WaterBody'
++BuiltinDisallowedActorClasses=/Script/CoreUObject.Class'/Script/Engine.Info'
++BuiltinAllowedComponentClasses=/Script/CoreUObject.Class'/Script/Engine.StaticMeshComponent'
++BuiltinDisallowedComponentClasses=/Script/CoreUObject.Class'/Script/Foliage.InteractiveFoliageComponent'
++BuiltinDisallowedComponentClasses=/Script/CoreUObject.Class'/Script/Landscape.LandscapeNaniteComponent'
++BuiltinDisallowedComponentClasses=/Script/CoreUObject.Class'/Script/Landscape.LandscapeMeshProxyComponent'
++BuiltinDisallowedComponentClasses=/Script/CoreUObject.Class'/Script/Water.WaterBodyMeshComponent'
++BuiltinDisallowedComponentClasses=/Script/CoreUObject.Class'/Script/Engine.SplineMeshComponent'
++BuiltinIgnoredRemainingComponentClasses=/Script/CoreUObject.Class'/Script/Engine.ActorTextureStreamingBuildDataComponent'
+
diff --git a/Engine/Plugins/Experimental/FastGeoStreaming/Config/FastGeoStreaming.ini b/Engine/Plugins/Experimental/FastGeoStreaming/Config/FastGeoStreaming.ini
deleted file mode 100644
