[FastGeo] Fixed cases where some FastGeo Container were registering to the Navigation System which was not yet initialized for the world.

diff --git a/Engine/Plugins/Experimental/FastGeoStreaming/Source/FastGeoStreaming/Internal/FastGeoContainer.h b/Engine/Plugins/Experimental/FastGeoStreaming/Source/FastGeoStreaming/Internal/FastGeoContainer.h
--- a/Engine/Plugins/Experimental/FastGeoStreaming/Source/FastGeoStreaming/Internal/FastGeoContainer.h
+++ b/Engine/Plugins/Experimental/FastGeoStreaming/Source/FastGeoStreaming/Internal/FastGeoContainer.h
@@ -22,6 +22,7 @@ class ULevel;
 class URuntimeVirtualTexture;
 class UStaticMesh;
 class UMaterialInterface;
+class UNavigationSystemBase;
 class FFastGeoComponent;
 class FChaosUserDefinedEntity;
 
@@ -194,6 +195,7 @@ private:
 	void SerializeComponentClusters(FArchive& Ar);
 	void RegisterToNavigationSystem();
 	void UnregisterFromNavigationSystem();
+	void OnNavigationInitDone(const UNavigationSystemBase& NavSys);
 
 	template <typename ThisType, typename TComponentCluster, typename TFunc>
 	static void ForEachComponentCluster(ThisType* Self, TFunc&& InFunc)
diff --git a/Engine/Plugins/Experimental/FastGeoStreaming/Source/FastGeoStreaming/Private/FastGeoContainer.cpp b/Engine/Plugins/Experimental/FastGeoStreaming/Source/FastGeoStreaming/Private/FastGeoContainer.cpp
--- a/Engine/Plugins/Experimental/FastGeoStreaming/Source/FastGeoStreaming/Private/FastGeoContainer.cpp
+++ b/Engine/Plugins/Experimental/FastGeoStreaming/Source/FastGeoStreaming/Private/FastGeoContainer.cpp
@@ -206,9 +206,18 @@ void UFastGeoContainer::Register()
 		}
 
 		WorldSubsystem->PushAsyncCreatePhysicsStateJobs(this);
-		RegisterToNavigationSystem();
 
 		bIsRegistered = true;
+
+		const UNavigationSystemBase* NavigationSystem = GetWorld()->GetNavigationSystem();
+		if (NavigationSystem && NavigationSystem->IsWorldInitDone())
+		{
+			RegisterToNavigationSystem();
+		}
+		else
+		{
+			UNavigationSystemBase::OnNavigationInitDoneStaticDelegate().AddUObject(this, &UFastGeoContainer::OnNavigationInitDone);
+		}
 	}
 
 	if (PendingCreate.HasAnyPendingState())
@@ -334,9 +343,27 @@ UWorld* UFastGeoContainer::GetWorld() const
 	return GetLevel()->GetWorld();
 }
 
+void UFastGeoContainer::OnNavigationInitDone(const UNavigationSystemBase& InNavigationSystem)
+{
+	const UWorld* World = GetWorld();
+	const UNavigationSystemBase* NavigationSystem = World ? World->GetNavigationSystem() : nullptr;
+
+	if (NavigationSystem && (NavigationSystem == &InNavigationSystem))
+	{
+		UNavigationSystemBase::OnNavigationInitDoneStaticDelegate().RemoveAll(this);
+		RegisterToNavigationSystem();
+	}
+}
+
 void UFastGeoContainer::RegisterToNavigationSystem()
 {
 	UWorld* World = GetWorld();
+
+	const UNavigationSystemBase* NavigationSystem = World->GetNavigationSystem();
+	check(NavigationSystem);
+	check(NavigationSystem->IsWorldInitDone());
+	check(bIsRegistered);
+
 	if (!FNavigationSystem::SupportsDynamicChanges(World))
 	{
 		return;
@@ -377,6 +404,8 @@ void UFastGeoContainer::RegisterToNavigationSystem()
 
 void UFastGeoContainer::UnregisterFromNavigationSystem()
 {
+	UNavigationSystemBase::OnNavigationInitDoneStaticDelegate().RemoveAll(this);
+
 	if (!NavigationElementHandles.IsEmpty())
 	{
 		TRACE_CPUPROFILER_EVENT_SCOPE(UFastGeoContainer::UnregisterFromNavigationSystem);
