[FastGeo] Small refactoring to facilitate readability in UFastGeoWorldSubsystem::AddToComponentsPendingRecreate.

diff --git a/Engine/Plugins/Experimental/FastGeoStreaming/Source/FastGeoStreaming/Private/FastGeoWorldSubsystem.cpp b/Engine/Plugins/Experimental/FastGeoStreaming/Source/FastGeoStreaming/Private/FastGeoWorldSubsystem.cpp
--- a/Engine/Plugins/Experimental/FastGeoStreaming/Source/FastGeoStreaming/Private/FastGeoWorldSubsystem.cpp
+++ b/Engine/Plugins/Experimental/FastGeoStreaming/Source/FastGeoStreaming/Private/FastGeoWorldSubsystem.cpp
@@ -326,23 +326,19 @@ void UFastGeoWorldSubsystem::ProcessPendingRecreate()
 	CSV_CUSTOM_STAT(FastGeo, PendingRecreate, ComponentsPendingRecreate.Num(), ECsvCustomStatOp::Set);
 
 	int32 NumPrimitiveDelayed = 0;
-	for (const FPendingRecreateComponent& PendingRecreateComponent : ComponentsPendingRecreate)
+	for (const FRegisteredComponent& PendingRecreateComponent : ComponentsPendingRecreate)
 	{
-		if (FFastGeoPrimitiveComponent* Component = PendingRecreateComponent.Component.Get<FFastGeoPrimitiveComponent>())
+		// Skip component if its container was unregistered or re-registered before we process it
+		if (FFastGeoPrimitiveComponent* Component = PendingRecreateComponent.TryGetRegistered<FFastGeoPrimitiveComponent>())
 		{
-			// Skip component if its container was unregistered or re-registered before we process it
-			if (Component->GetOwnerContainer()->IsRegistered() &&
-				(PendingRecreateComponent.RegistrationEpoch == Component->GetOwnerContainer()->RegistrationEpoch))
-			{
 #if CSV_PROFILER && !CSV_PROFILER_MINIMAL
-				if (Component->IsRenderStateDelayed())
-				{
-					++NumPrimitiveDelayed;
-				}
-#endif
-				Component->DestroyRenderState(nullptr);
-				Component->CreateRenderState(nullptr);
+			if (Component->IsRenderStateDelayed())
+			{
+				++NumPrimitiveDelayed;
 			}
+#endif
+			Component->DestroyRenderState(nullptr);
+			Component->CreateRenderState(nullptr);
 		}
 	}
 
diff --git a/Engine/Plugins/Experimental/FastGeoStreaming/Source/FastGeoStreaming/Private/FastGeoWorldSubsystem.h b/Engine/Plugins/Experimental/FastGeoStreaming/Source/FastGeoStreaming/Private/FastGeoWorldSubsystem.h
--- a/Engine/Plugins/Experimental/FastGeoStreaming/Source/FastGeoStreaming/Private/FastGeoWorldSubsystem.h
+++ b/Engine/Plugins/Experimental/FastGeoStreaming/Source/FastGeoStreaming/Private/FastGeoWorldSubsystem.h
@@ -98,9 +98,9 @@ private:
 	FDelegateHandle Handle_OnLevelBeginRemoveFromWorld;
 	FDelegateHandle Handle_OnForEachHLODObjectInCell;
 
-	struct FPendingRecreateComponent
+	struct FRegisteredComponent
 	{
-		FPendingRecreateComponent(FFastGeoComponent* InComponent)
+		FRegisteredComponent(FFastGeoComponent* InComponent)
 			: Component(InComponent)
 		{
 			check(InComponent);
@@ -108,16 +108,31 @@ private:
 			RegistrationEpoch = InComponent->GetOwnerContainer()->RegistrationEpoch;
 		}
 
-		FWeakFastGeoComponent Component;
-		int32 RegistrationEpoch = -1;
+		template <class T>
+		FORCEINLINE T* TryGetRegistered() const
+		{
+			if (T* ComponentPtr = Component.Get<T>())
+			{
+				if (ComponentPtr->GetOwnerContainer()->IsRegistered() &&
+					RegistrationEpoch == ComponentPtr->GetOwnerContainer()->RegistrationEpoch)
+				{
+					return ComponentPtr;
+				}
+			}
+			return nullptr;
+		}
 
-		FORCEINLINE bool operator==(const FPendingRecreateComponent& Other) const
+		FORCEINLINE bool operator==(const FRegisteredComponent& Other) const
 		{
 			return Component == Other.Component && RegistrationEpoch == Other.RegistrationEpoch;
 		}
+
+	private:
+		FWeakFastGeoComponent Component;
+		int32 RegistrationEpoch = -1;
 	};
 
-	TArray<FPendingRecreateComponent> ComponentsPendingRecreate;
+	TArray<FRegisteredComponent> ComponentsPendingRecreate;
 	static bool bEnableDebugView;
 
 	FRWLock Lock;
