[FastGeo] Removed unused UFastGeoComponentEditorProxy.

diff --git a/Engine/Plugins/Experimental/FastGeoStreaming/Source/FastGeoStreaming/Internal/FastGeoComponent.h b/Engine/Plugins/Experimental/FastGeoStreaming/Source/FastGeoStreaming/Internal/FastGeoComponent.h
--- a/Engine/Plugins/Experimental/FastGeoStreaming/Source/FastGeoStreaming/Internal/FastGeoComponent.h
+++ b/Engine/Plugins/Experimental/FastGeoStreaming/Source/FastGeoStreaming/Internal/FastGeoComponent.h
@@ -7,7 +7,6 @@
 #include "FastGeoElementType.h"
 #include "Physics/Experimental/AsyncPhysicsStateProcessorInterface.h"
 #include "Components/ComponentInterfaces.h"
-#include "FastGeoComponent.generated.h"
 
 class UWorld;
 class UBodySetup;
@@ -56,19 +55,6 @@ public:
 	bool IsRegistered() const;
 	FLinearColor GetDebugColor() const;
 
-#if WITH_EDITOR
-	void SetEditorProxy(UFastGeoComponentEditorProxy* InComponentEditorProxy)
-	{
-		ComponentEditorProxy = InComponentEditorProxy;
-	}
-		
-	template <typename TEditorProxy>
-	TEditorProxy* GetEditorProxy() const
-	{
-		return Cast<TEditorProxy>(ComponentEditorProxy);
-	}
-#endif
-
 protected:
 	// Persistent Data
 	int32 ComponentIndex;
@@ -91,46 +77,4 @@ private:
 	friend class FFastGeoComponentCluster;
 
 	friend FArchive& operator<<(FArchive& Ar, FFastGeoComponent& Component);
-
-#if WITH_EDITOR
-	UFastGeoComponentEditorProxy* ComponentEditorProxy;
-#endif
-};
-
-
-UCLASS()
-class UFastGeoComponentEditorProxy : public UObject
-{
-	GENERATED_BODY()
-
-#if WITH_EDITOR
-	friend class UFastGeoContainer;
-
-protected:
-	template <typename TComponent>
-	const TComponent& GetComponent() const
-	{
-		return FastGeoComponent->CastToRef<TComponent>();
-	}
-
-	template <typename TComponent>
-	TComponent& GetComponent()
-	{
-		return FastGeoComponent->CastToRef<TComponent>();
-	}
-
-private:
-	void SetFastGeoComponent(FFastGeoComponent* InFastGeoComponent)
-	{
-		FastGeoComponent = InFastGeoComponent;
-	}
-
-	const FFastGeoComponent* GetFastGeoComponent() const
-	{
-		return FastGeoComponent;
-	}
-
-private:
-	FFastGeoComponent* FastGeoComponent;
-#endif
-};
+};
\ No newline at end of file
diff --git a/Engine/Plugins/Experimental/FastGeoStreaming/Source/FastGeoStreaming/Public/FastGeoStreamingModule.h b/Engine/Plugins/Experimental/FastGeoStreaming/Source/FastGeoStreaming/Public/FastGeoStreamingModule.h
--- a/Engine/Plugins/Experimental/FastGeoStreaming/Source/FastGeoStreaming/Public/FastGeoStreamingModule.h
+++ b/Engine/Plugins/Experimental/FastGeoStreaming/Source/FastGeoStreaming/Public/FastGeoStreamingModule.h
@@ -5,7 +5,6 @@
 #include "Delegates/Delegate.h"
 #include "Modules/ModuleInterface.h"
 
-class FFastGeoPrimitiveComponent;
 class UWorld;
 
 class FFastGeoStreamingModule : public IModuleInterface
diff --git a/Engine/Plugins/Experimental/FastGeoStreaming/Source/FastGeoStreamingEditor/Private/FastGeoStreamingEditorModule.h b/Engine/Plugins/Experimental/FastGeoStreaming/Source/FastGeoStreamingEditor/Private/FastGeoStreamingEditorModule.h
--- a/Engine/Plugins/Experimental/FastGeoStreaming/Source/FastGeoStreamingEditor/Private/FastGeoStreamingEditorModule.h
+++ b/Engine/Plugins/Experimental/FastGeoStreaming/Source/FastGeoStreamingEditor/Private/FastGeoStreamingEditorModule.h
@@ -5,9 +5,6 @@
 #include "Delegates/Delegate.h"
 #include "Modules/ModuleInterface.h"
 
-class FFastGeoPrimitiveComponent;
-class UWorld;
-
 class FFastGeoStreamingEditorModule : public IModuleInterface
 {
 public:
