[FastGeo] Preserve root component mobility when FastGeo runtime cell transformer creates a new replacement root component for a transformed component.

diff --git a/Engine/Plugins/Experimental/FastGeoStreaming/Source/FastGeoStreaming/Private/FastGeoWorldPartitionRuntimeCellTransformer.cpp b/Engine/Plugins/Experimental/FastGeoStreaming/Source/FastGeoStreaming/Private/FastGeoWorldPartitionRuntimeCellTransformer.cpp
--- a/Engine/Plugins/Experimental/FastGeoStreaming/Source/FastGeoStreaming/Private/FastGeoWorldPartitionRuntimeCellTransformer.cpp
+++ b/Engine/Plugins/Experimental/FastGeoStreaming/Source/FastGeoStreaming/Private/FastGeoWorldPartitionRuntimeCellTransformer.cpp
@@ -354,6 +354,7 @@ void UFastGeoWorldPartitionRuntimeCellTransformer::Transform(ULevel* InLevel)
 				// Replace removed root component with a SceneComponent and remap attachment of other components
 				USceneComponent* NewRootComponent = NewObject<USceneComponent>(Actor);
 				NewRootComponent->SetRelativeTransform(OldRootComponent->GetRelativeTransform());
+				NewRootComponent->SetMobility(OldRootComponent->GetMobility());
 				Actor->SetRootComponent(NewRootComponent);
 
 				Actor->ForEachComponent<USceneComponent>(false, [OldRootComponent, NewRootComponent](USceneComponent* Component)
