[FastGeo] Added support for DecalComponent.
- Adapted  FastGeo to support non-primitive components.
- Added FFastGeoPSOPrecacheFinishedTask to share handling PSO precaching task completion
- Moved code from FFastGeoPrimitiveComponent back to base class FFastGeoComponent

diff --git a/Engine/Plugins/Experimental/FastGeoStreaming/Config/DefaultFastGeoStreaming.ini b/Engine/Plugins/Experimental/FastGeoStreaming/Config/DefaultFastGeoStreaming.ini
--- a/Engine/Plugins/Experimental/FastGeoStreaming/Config/DefaultFastGeoStreaming.ini
+++ b/Engine/Plugins/Experimental/FastGeoStreaming/Config/DefaultFastGeoStreaming.ini
@@ -6,6 +6,7 @@
 +BuiltinDisallowedActorClasses=/Script/CoreUObject.Class'/Script/Water.WaterBody'
 +BuiltinDisallowedActorClasses=/Script/CoreUObject.Class'/Script/Engine.Info'
 +BuiltinAllowedComponentClasses=/Script/CoreUObject.Class'/Script/Engine.StaticMeshComponent'
++BuiltinAllowedComponentClasses=/Script/CoreUObject.Class'/Script/Engine.DecalComponent'
 +BuiltinDisallowedComponentClasses=/Script/CoreUObject.Class'/Script/Foliage.InteractiveFoliageComponent'
 +BuiltinDisallowedComponentClasses=/Script/CoreUObject.Class'/Script/Landscape.LandscapeNaniteComponent'
 +BuiltinDisallowedComponentClasses=/Script/CoreUObject.Class'/Script/Landscape.LandscapeMeshProxyComponent'
diff --git a/Engine/Plugins/Experimental/FastGeoStreaming/Source/FastGeoStreaming/Internal/FastGeoComponent.h b/Engine/Plugins/Experimental/FastGeoStreaming/Source/FastGeoStreaming/Internal/FastGeoComponent.h
--- a/Engine/Plugins/Experimental/FastGeoStreaming/Source/FastGeoStreaming/Internal/FastGeoComponent.h
+++ b/Engine/Plugins/Experimental/FastGeoStreaming/Source/FastGeoStreaming/Internal/FastGeoComponent.h
@@ -3,20 +3,31 @@
 #pragma once
 
 #include "CoreMinimal.h"
+#include "PSOPrecacheFwd.h"
 #include "IFastGeoElement.h"
 #include "FastGeoElementType.h"
 #include "Physics/Experimental/AsyncPhysicsStateProcessorInterface.h"
 #include "Components/ComponentInterfaces.h"
+#include "Async/TaskGraphInterfaces.h"
 
 class UWorld;
 class UBodySetup;
 class UActorComponent;
 class FFastGeoComponentCluster;
+class FFastGeoDestroyRenderStateContext;
+class FWeakFastGeoComponent;
 class UFastGeoContainer;
 class FSceneInterface;
 class FPrimitiveSceneProxy;
 class FRegisterComponentContext;
 
+enum class ESceneProxyCreationError
+{
+	None,
+	WaitingPSOs,
+	InvalidMesh
+};
+
 #define FASTGEO_DECLARE_COMPONENT_INTERFACE_INTERNAL(ComponentName, BasePrefix, ...)\
 	UE_DECLARE_COMPONENT_INTERFACE_INTERNAL(BasePrefix##ComponentName##__VA_ARGS__, I##ComponentName , BasePrefix##ComponentName##Interface, ComponentName##Interface, ComponentName)
 
@@ -32,17 +43,31 @@ public:
 	static const FFastGeoElementType Type;
 
 	FFastGeoComponent(int32 InComponentIndex = INDEX_NONE, FFastGeoElementType InType = Type);
+	FFastGeoComponent(const FFastGeoComponent& Other);
 	virtual ~FFastGeoComponent() = default;
 
 #if WITH_EDITOR
 	virtual void InitializeFromComponent(UActorComponent* Component);
 #endif
-
 	int32 GetComponentIndex() const { return ComponentIndex; }
+	void MarkRenderStateDirty();
+	bool IsRenderStateDirty() const;
+	bool IsRenderStateDelayed() const;
+	bool IsRenderStateCreated() const;
 	virtual UBodySetup* GetBodySetup() const { return nullptr; }
 	virtual bool IsCollisionEnabled() const { return false; }
 	virtual void Serialize(FArchive& Ar);
 	virtual void InitializeDynamicProperties() {}
+	virtual void ApplyWorldTransform(const FTransform& InTransform) {}
+	virtual bool ShouldCreateRenderState() const { return false; }
+	virtual void CreateRenderState(FRegisterComponentContext* Context) {}
+	virtual void DestroyRenderState(FFastGeoDestroyRenderStateContext* Context) {}
+	virtual void UpdateVisibility() {}
+	virtual void MarkPrecachePSOsRequired() {}
+	virtual void PrecachePSOs() {}
+#if !WITH_EDITOR && UE_WITH_PSO_PRECACHING
+	virtual void OnPrecacheFinished(int32 JobSetThatJustCompleted) {}
+#endif
 	virtual void OnAsyncCreatePhysicsState();
 	virtual void OnAsyncDestroyPhysicsStateBegin_GameThread();
 	virtual void OnAsyncDestroyPhysicsState();
@@ -52,6 +77,7 @@ public:
 	FFastGeoComponentCluster* GetOwnerComponentCluster() const;
 	UFastGeoContainer* GetOwnerContainer() const;
 	UWorld* GetWorld() const;
+	FSceneInterface* GetScene() const;
 	bool IsRegistered() const;
 	FLinearColor GetDebugColor() const;
 
@@ -72,9 +98,20 @@ protected:
 
 	EPhysicsStateCreation PhysicsStateCreation = EPhysicsStateCreation::NotCreated;
 
+	enum class EProxyCreationState : uint8
+	{
+		None, // Constructed/Initialized
+		Creating, // Actively creating the proxy
+		Created, // Proxy is now created
+		Delayed // Proxy creation delayed (used when PSO precaching is not ready when creating proxy)
+	};
+
+	EProxyCreationState ProxyState = EProxyCreationState::None;
+	bool bRenderStateDirty = false;
+
 private:
 	void SetOwnerComponentCluster(FFastGeoComponentCluster* InOwner);
 	friend class FFastGeoComponentCluster;
 
 	friend FArchive& operator<<(FArchive& Ar, FFastGeoComponent& Component);
-};
\ No newline at end of file
+};
diff --git a/Engine/Plugins/Experimental/FastGeoStreaming/Source/FastGeoStreaming/Internal/FastGeoComponentCluster.h b/Engine/Plugins/Experimental/FastGeoStreaming/Source/FastGeoStreaming/Internal/FastGeoComponentCluster.h
--- a/Engine/Plugins/Experimental/FastGeoStreaming/Source/FastGeoStreaming/Internal/FastGeoComponentCluster.h
+++ b/Engine/Plugins/Experimental/FastGeoStreaming/Source/FastGeoStreaming/Internal/FastGeoComponentCluster.h
@@ -9,6 +9,7 @@
 #include "FastGeoSkinnedMeshComponent.h"
 #include "FastGeoInstancedSkinnedMeshComponent.h"
 #include "FastGeoProceduralISMComponent.h"
+#include "FastGeoDecalComponent.h"
 
 class ULevel;
 class UFastGeoContainer;
@@ -88,6 +89,7 @@ protected:
 	TArray<FFastGeoSkinnedMeshComponent> SkinnedMeshComponents;
 	TArray<FFastGeoInstancedSkinnedMeshComponent> InstancedSkinnedMeshComponents;
 	TArray<FFastGeoProceduralISMComponent> ProceduralISMComponents;
+	TArray<FFastGeoDecalComponent> DecalComponents;
 
 	template <typename ThisType, typename TComponent, typename TFunc>
 	static void ForEachComponent(ThisType* Self, TFunc&& InFunc)
@@ -119,6 +121,7 @@ protected:
 		ForEachArray(Self->SkinnedMeshComponents);
 		ForEachArray(Self->InstancedSkinnedMeshComponents);
 		ForEachArray(Self->ProceduralISMComponents);
+		ForEachArray(Self->DecalComponents);
 	}
 
 	template <typename ThisType, typename TComponent, typename TFunc>
@@ -158,7 +161,8 @@ protected:
 			   ForEachArray(Self->InstancedStaticMeshComponents) &&
 			   ForEachArray(Self->SkinnedMeshComponents) &&
 			   ForEachArray(Self->InstancedSkinnedMeshComponents) &&
-			   ForEachArray(Self->ProceduralISMComponents);
+			   ForEachArray(Self->ProceduralISMComponents) && 
+			   ForEachArray(Self->DecalComponents);
 	}
 
 private:
diff --git a/Engine/Plugins/Experimental/FastGeoStreaming/Source/FastGeoStreaming/Internal/FastGeoContainer.h b/Engine/Plugins/Experimental/FastGeoStreaming/Source/FastGeoStreaming/Internal/FastGeoContainer.h
--- a/Engine/Plugins/Experimental/FastGeoStreaming/Source/FastGeoStreaming/Internal/FastGeoContainer.h
+++ b/Engine/Plugins/Experimental/FastGeoStreaming/Source/FastGeoStreaming/Internal/FastGeoContainer.h
@@ -24,6 +24,7 @@ class UStaticMesh;
 class UMaterialInterface;
 class UNavigationSystemBase;
 class FFastGeoComponent;
+class FFastGeoComponentCluster;
 class FChaosUserDefinedEntity;
 
 struct FFastGeoComponentClusterPendingState
@@ -48,7 +49,7 @@ struct FFastGeoComponentClusterPendingState
 			return TotalNumProcessed >= ComponentsToProcess.Num();
 		}
 
-		TArray<FFastGeoPrimitiveComponent*> ComponentsToProcess;
+		TArray<FFastGeoComponent*> ComponentsToProcess;
 
 		int32 NumToProcess;
 		int32 NumProcessed;
@@ -296,5 +297,5 @@ private:
 	FRenderCommandFence DestroyFence;
 
 	friend class FFastGeoGatherFastGeoContainerAssetRefsArchive;
-	friend class UFastGeoWorldSubsystem;
+	friend struct FFastGeoRegisteredComponent;
 };
\ No newline at end of file
diff --git a/Engine/Plugins/Experimental/FastGeoStreaming/Source/FastGeoStreaming/Internal/FastGeoDecalComponent.h b/Engine/Plugins/Experimental/FastGeoStreaming/Source/FastGeoStreaming/Internal/FastGeoDecalComponent.h
new file mode 100644
--- /dev/null
+++ b/Engine/Plugins/Experimental/FastGeoStreaming/Source/FastGeoStreaming/Internal/FastGeoDecalComponent.h
@@ -0,0 +1,63 @@
+// Copyright Epic Games, Inc. All Rights Reserved.
+
+#pragma once
+
+#include "FastGeoComponent.h"
+#include "DeferredDecalSceneProxyDesc.h"
+#include "Engine/TimerHandle.h"
+
+class FDeferredDecalProxy;
+
+class FASTGEOSTREAMING_API FFastGeoDecalComponent : public FFastGeoComponent
+{
+public:
+	typedef FFastGeoComponent Super;
+
+	/** Static type identifier for this element class */
+	static const FFastGeoElementType Type;
+
+	FFastGeoDecalComponent(int32 InComponentIndex = INDEX_NONE, FFastGeoElementType InType = Type);
+	FFastGeoDecalComponent(const FFastGeoDecalComponent& Other);
+	virtual ~FFastGeoDecalComponent() = default;
+
+	//~ Begin FFastGeoComponent interface
+	virtual void Serialize(FArchive& Ar) override;
+	virtual void ApplyWorldTransform(const FTransform& InTransform) override;
+	virtual bool ShouldCreateRenderState() const override;
+	virtual void CreateRenderState(FRegisterComponentContext* Context) override;
+	virtual void DestroyRenderState(FFastGeoDestroyRenderStateContext* Context) override;
+	virtual void UpdateVisibility() override;
+	virtual void PrecachePSOs() override;
+#if !WITH_EDITOR && UE_WITH_PSO_PRECACHING
+	virtual void OnPrecacheFinished(int32 JobSetThatJustCompleted) override;
+#endif
+#if WITH_EDITOR
+	virtual void InitializeFromComponent(UActorComponent* Component) override;
+#endif
+	//~ End FFastGeoComponent interface
+
+protected:
+	FBoxSphereBounds CalcBounds() const;
+	FTransform GetTransformIncludingDecalSize() const;
+	FDeferredDecalProxy* CreateSceneProxy(ESceneProxyCreationError& OutError);
+	void InitializeSceneProxyDescDynamicProperties();
+	void UpdateVisibilityInternal();
+	void SetLifeSpan(const float InLifeSpan);
+
+private:
+	// Persistent data
+	FTransform LocalTransform;
+	FTransform WorldTransform;
+	FVector DecalSize;
+	uint8 bIsVisible : 1 = true;
+	FDeferredDecalSceneProxyDesc SceneProxyDesc{};
+	FDeferredDecalProxy* SceneProxy = nullptr;
+	TUniquePtr<FRWLock> Lock;
+	FTimerHandle TimerHandle_DestroyDecalComponent;
+
+#if !WITH_EDITOR && UE_WITH_PSO_PRECACHING
+	// Atomic int used to track the last PSO precache events
+	std::atomic<int> LatestPSOPrecacheJobSetCompleted = 0;
+	int32 LatestPSOPrecacheJobSet = 0;
+#endif
+};
\ No newline at end of file
diff --git a/Engine/Plugins/Experimental/FastGeoStreaming/Source/FastGeoStreaming/Internal/FastGeoDestroyRenderStateContext.h b/Engine/Plugins/Experimental/FastGeoStreaming/Source/FastGeoStreaming/Internal/FastGeoDestroyRenderStateContext.h
new file mode 100644
--- /dev/null
+++ b/Engine/Plugins/Experimental/FastGeoStreaming/Source/FastGeoStreaming/Internal/FastGeoDestroyRenderStateContext.h
@@ -0,0 +1,24 @@
+// Copyright Epic Games, Inc. All Rights Reserved.
+
+#pragma once
+
+#include "Containers/Array.h"
+
+class FSceneInterface;
+class FDeferredDecalProxy;
+class FPrimitiveSceneProxy;
+
+class FFastGeoDestroyRenderStateContext
+{
+public:
+	FFastGeoDestroyRenderStateContext(FSceneInterface* InScene);
+	~FFastGeoDestroyRenderStateContext();
+
+	static void DestroyProxy(FFastGeoDestroyRenderStateContext* InContext, FPrimitiveSceneProxy* InPrimitiveSceneProxy, FSceneInterface* InScene);
+	static void DestroyProxy(FFastGeoDestroyRenderStateContext* InContext, FDeferredDecalProxy* InDeferredDecalProxy, FSceneInterface* InScene);
+
+private:
+	FSceneInterface* Scene;
+	TArray<FPrimitiveSceneProxy*> PrimitiveSceneProxies;
+	TArray<FDeferredDecalProxy*> DeferredDecalProxies;
+};
\ No newline at end of file
diff --git a/Engine/Plugins/Experimental/FastGeoStreaming/Source/FastGeoStreaming/Internal/FastGeoPSOPrecacheFinishedTask.h b/Engine/Plugins/Experimental/FastGeoStreaming/Source/FastGeoStreaming/Internal/FastGeoPSOPrecacheFinishedTask.h
new file mode 100644
--- /dev/null
+++ b/Engine/Plugins/Experimental/FastGeoStreaming/Source/FastGeoStreaming/Internal/FastGeoPSOPrecacheFinishedTask.h
@@ -0,0 +1,34 @@
+// Copyright Epic Games, Inc. All Rights Reserved.
+
+#pragma once
+
+#include "FastGeoWeakElement.h"
+
+#if !WITH_EDITOR && UE_WITH_PSO_PRECACHING
+struct FFastGeoPSOPrecacheFinishedTask
+{
+	explicit FFastGeoPSOPrecacheFinishedTask(FFastGeoComponent* InComponent, int32 InJobSetThatJustCompleted)
+		: WeakComponent(InComponent)
+		, JobSetThatJustCompleted(InJobSetThatJustCompleted)
+	{
+	}
+
+	static TStatId GetStatId() { return TStatId(); }
+	static ENamedThreads::Type GetDesiredThread() { return ENamedThreads::GameThread; }
+	static ESubsequentsMode::Type GetSubsequentsMode() { return ESubsequentsMode::TrackSubsequents; }
+	void DoTask(ENamedThreads::Type CurrentThread, const FGraphEventRef& MyCompletionGraphEvent)
+	{
+		if (FFastGeoComponent* Component = WeakComponent.Get())
+		{
+			// Validate that the component is still part of a streamed-in level
+			if (Component->GetWorld())
+			{
+				Component->OnPrecacheFinished(JobSetThatJustCompleted);
+			}
+		}
+	}
+
+	FWeakFastGeoComponent WeakComponent;
+	int32 JobSetThatJustCompleted;
+};
+#endif
\ No newline at end of file
diff --git a/Engine/Plugins/Experimental/FastGeoStreaming/Source/FastGeoStreaming/Internal/FastGeoPrimitiveComponent.h b/Engine/Plugins/Experimental/FastGeoStreaming/Source/FastGeoStreaming/Internal/FastGeoPrimitiveComponent.h
--- a/Engine/Plugins/Experimental/FastGeoStreaming/Source/FastGeoStreaming/Internal/FastGeoPrimitiveComponent.h
+++ b/Engine/Plugins/Experimental/FastGeoStreaming/Source/FastGeoStreaming/Internal/FastGeoPrimitiveComponent.h
@@ -12,25 +12,20 @@
 #include "PhysicsEngine/BodyInstance.h"
 #include "PhysicsEngine/PhysicsBodyInstanceOwnerInterface.h"
 #include "Chaos/ChaosUserEntity.h"
-#include "PSOPrecacheFwd.h"
 #include "Templates/DontCopy.h"
 
+class FRegisterComponentContext;
 class FPrimitiveSceneProxy;
 class UPrimitiveComponent;
 class URuntimeVirtualTexture;
+class UFastGeoContainer;
 class FFastGeoPrimitiveComponent;
+class FFastGeoDestroyRenderStateContext;
 struct FPrimitiveSceneDesc;
 struct FPrimitiveSceneProxyDesc;
 struct FPSOPrecacheParams;
 enum class EPSOPrecachePriority : uint8;
 
-enum class ESceneProxyCreationError
-{
-	None,
-	WaitingPSOs,
-	InvalidMesh
-};
-
 class FFastGeoPhysicsBodyInstanceOwner : public FChaosUserDefinedEntity, public IPhysicsBodyInstanceOwner
 {
 public:
@@ -83,7 +78,6 @@ public:
 	FPrimitiveComponentId GetPrimitiveSceneId() const { return PrimitiveSceneData.PrimitiveSceneId; }
 	virtual UObject const* AdditionalStatObject() const { return nullptr; }
 		
-	FSceneInterface* GetScene() const;
 	FPrimitiveSceneProxy* GetSceneProxy() const;
 	const FTransform& GetTransform() const;
 	const FBoxSphereBounds& GetBounds() const;
@@ -93,33 +87,10 @@ public:
 	// Used by FStaticMeshComponentHelper/FInstancedStaticMeshComponentHelper
 	const FTransform& GetComponentTransform() const { return GetTransform(); }
 
-	void CreateRenderState(FRegisterComponentContext* Context);
-
-	struct FFastGeoDestroyRenderStateContext
-	{
-		FFastGeoDestroyRenderStateContext(FSceneInterface* InScene);
-		~FFastGeoDestroyRenderStateContext();
-		bool HasPendingWork() const;
-
-		static void DestroyProxy(FFastGeoDestroyRenderStateContext* InContext, FPrimitiveSceneProxy* InPrimitiveSceneProxy);
-
-	private:
-		FSceneInterface* Scene;
-		TArray<FPrimitiveSceneProxy*> PrimitiveSceneProxies;
-	};
-
-	virtual void DestroyRenderState(FFastGeoDestroyRenderStateContext* Context);
 #if WITH_EDITOR
 	virtual void NotifyRenderStateChanged();
 #endif
-	bool IsRenderStateCreated() const;
-	bool IsRenderStateDelayed() const;
-	bool IsRenderStateDirty() const;
 	bool IsDrawnInGame() const;
-	bool ShouldCreateRenderState() const;
-	void MarkRenderStateDirty();
-	void MarkPrecachePSOsRequired();
-	void PrecachePSOs();
 	bool CheckPSOPrecachingAndBoostPriority(EPSOPrecachePriority NewPSOPrecachePriority = EPSOPrecachePriority::High);
 	bool IsPSOPrecaching() const;
 	void SetCollisionEnabled(bool bInCollisionEnabled);
@@ -134,6 +105,12 @@ public:
 	virtual void OnAsyncDestroyPhysicsStateEnd_GameThread() override;
 	virtual void OnAsyncDestroyPhysicsState() override;
 	virtual bool IsCollisionEnabled() const override;
+	virtual void ApplyWorldTransform(const FTransform& InTransform) override;
+	virtual bool ShouldCreateRenderState() const override;
+	virtual void CreateRenderState(FRegisterComponentContext* Context) override;
+	virtual void DestroyRenderState(FFastGeoDestroyRenderStateContext* Context) override;
+	virtual void MarkPrecachePSOsRequired() override;
+	virtual void PrecachePSOs() override;
 #if WITH_EDITOR
 	virtual void InitializeFromComponent(UActorComponent* Component) override;
 #endif
@@ -166,7 +143,6 @@ protected:
 	virtual void ResetSceneProxyDescUnsupportedProperties();
 #endif
 	virtual void InitializeSceneProxyDescDynamicProperties();
-	virtual void ApplyWorldTransform(const FTransform& InTransform);
 	virtual FPrimitiveSceneProxy* CreateSceneProxy(ESceneProxyCreationError* OutError = nullptr) = 0;
 	FPrimitiveSceneDesc BuildSceneDesc();
 
@@ -201,23 +177,12 @@ protected:
 	/** Payload used to release BodyInstance resources in asynchronous mode (see OnAsyncDestroyPhysicsState). */
 	TOptional<FBodyInstance::FAsyncTermBodyPayload> AsyncTermBodyPayload;
 
-	enum class EProxyCreationState : uint8
-	{
-		None, // Constructed/Initialized
-		Creating, // Actively creating the proxy
-		Created, // Proxy is now created
-		Delayed // Proxy creation delayed (used when PSO precaching is not ready when creating proxy)
-	};
-
-	EProxyCreationState ProxyState = EProxyCreationState::None;
-	bool bRenderStateDirty = false;
-
 private:
 	bool ShouldRenderProxyFallbackToDefaultMaterial() const;
 
 #if !WITH_EDITOR && UE_WITH_PSO_PRECACHING
 public:
-	void OnPrecacheFinished(int32 JobSetThatJustCompleted);
+	virtual void OnPrecacheFinished(int32 JobSetThatJustCompleted) override;
 
 protected:
 	void RequestRecreateRenderStateWhenPSOPrecacheFinished(const FGraphEventArray& PSOPrecacheCompileEvents);
diff --git a/Engine/Plugins/Experimental/FastGeoStreaming/Source/FastGeoStreaming/Private/FastGeoComponent.cpp b/Engine/Plugins/Experimental/FastGeoStreaming/Source/FastGeoStreaming/Private/FastGeoComponent.cpp
--- a/Engine/Plugins/Experimental/FastGeoStreaming/Source/FastGeoStreaming/Private/FastGeoComponent.cpp
+++ b/Engine/Plugins/Experimental/FastGeoStreaming/Source/FastGeoStreaming/Private/FastGeoComponent.cpp
@@ -17,6 +17,16 @@ FFastGeoComponent::FFastGeoComponent(int32 InComponentIndex, FFastGeoElementType
 {
 }
 
+FFastGeoComponent::FFastGeoComponent(const FFastGeoComponent& Other)
+: Super(Other)
+, ComponentIndex(Other.ComponentIndex)
+, Owner(Other.Owner)
+, PhysicsStateCreation(EPhysicsStateCreation::NotCreated)
+, ProxyState(EProxyCreationState::None)
+, bRenderStateDirty(false)
+{
+}
+
 #if WITH_EDITOR
 void FFastGeoComponent::InitializeFromComponent(UActorComponent* InComponent)
 {
@@ -54,6 +64,15 @@ UWorld* FFastGeoComponent::GetWorld() const
 	return GetOwnerContainer()->GetLevel()->OwningWorld;
 }
 
+FSceneInterface* FFastGeoComponent::GetScene() const
+{
+	if (UWorld* World = GetWorld())
+	{
+		return World->Scene;
+	}
+	return nullptr;
+}
+
 bool FFastGeoComponent::IsRegistered() const
 {
 	return GetOwnerContainer()->IsRegistered();
@@ -98,3 +117,32 @@ void FFastGeoComponent::OnAsyncDestroyPhysicsStateEnd_GameThread()
 	check(PhysicsStateCreation == EPhysicsStateCreation::Destroying);
 	PhysicsStateCreation = EPhysicsStateCreation::NotCreated;
 }
+
+bool FFastGeoComponent::IsRenderStateCreated() const
+{
+	return ProxyState == EProxyCreationState::Created;
+}
+
+bool FFastGeoComponent::IsRenderStateDelayed() const
+{
+	return ProxyState == EProxyCreationState::Delayed;
+}
+
+bool FFastGeoComponent::IsRenderStateDirty() const
+{
+	return bRenderStateDirty;
+}
+
+void FFastGeoComponent::MarkRenderStateDirty()
+{
+	if ((IsRenderStateCreated() || IsRenderStateDelayed()) && !IsRenderStateDirty())
+	{
+		bRenderStateDirty = true;
+
+		UFastGeoWorldSubsystem* WorldSubsystem = GetWorld()->GetSubsystem<UFastGeoWorldSubsystem>();
+		if (ensure(WorldSubsystem))
+		{
+			WorldSubsystem->AddToComponentsPendingRecreate(this);
+		}
+	}
+}
\ No newline at end of file
diff --git a/Engine/Plugins/Experimental/FastGeoStreaming/Source/FastGeoStreaming/Private/FastGeoComponentCluster.cpp b/Engine/Plugins/Experimental/FastGeoStreaming/Source/FastGeoStreaming/Private/FastGeoComponentCluster.cpp
--- a/Engine/Plugins/Experimental/FastGeoStreaming/Source/FastGeoStreaming/Private/FastGeoComponentCluster.cpp
+++ b/Engine/Plugins/Experimental/FastGeoStreaming/Source/FastGeoStreaming/Private/FastGeoComponentCluster.cpp
@@ -27,6 +27,7 @@ FFastGeoComponentCluster::FFastGeoComponentCluster(const FFastGeoComponentCluste
 	SkinnedMeshComponents = InOther.SkinnedMeshComponents;
 	InstancedSkinnedMeshComponents = InOther.InstancedSkinnedMeshComponents;
 	ProceduralISMComponents = InOther.ProceduralISMComponents;
+	DecalComponents = InOther.DecalComponents;
 }
 
 FArchive& operator<<(FArchive& Ar, FFastGeoComponentCluster& ComponentCluster)
@@ -44,6 +45,7 @@ void FFastGeoComponentCluster::Serialize(FArchive& Ar)
 	Ar << SkinnedMeshComponents;
 	Ar << InstancedSkinnedMeshComponents;
 	Ar << ProceduralISMComponents;
+	Ar << DecalComponents;
 }
 
 void FFastGeoComponentCluster::InitializeDynamicProperties()
@@ -93,6 +95,10 @@ FFastGeoComponent* FFastGeoComponentCluster::GetComponent(uint32 InComponentType
 	{
 		return ProceduralISMComponents.IsValidIndex(InComponentIndex) ? &ProceduralISMComponents[InComponentIndex] : nullptr;
 	}
+	else if (FFastGeoDecalComponent::Type.IsSameTypeID(InComponentTypeID))
+	{
+		return DecalComponents.IsValidIndex(InComponentIndex) ? &DecalComponents[InComponentIndex] : nullptr;
+	}
 	check(false);
 	return nullptr;
 }
@@ -108,26 +114,33 @@ void FFastGeoComponentCluster::UpdateVisibility()
 
 	TArray<FFastGeoPrimitiveComponent*> ShowComponents;
 	TArray<FFastGeoPrimitiveComponent*> HideComponents;
-	ForEachComponent<FFastGeoPrimitiveComponent>([&ShowComponents, &HideComponents](FFastGeoPrimitiveComponent& Component)
+	ForEachComponent<FFastGeoComponent>([&ShowComponents, &HideComponents](FFastGeoComponent& Component)
 	{
-		FWriteScopeLock WriteLock(*Component.Lock.Get());
-		const bool bOldIsDrawnInGame = Component.IsDrawnInGame();
-		Component.UpdateVisibility();
-		const bool bIsDrawnInGame = Component.IsDrawnInGame();
-		if (bIsDrawnInGame != bOldIsDrawnInGame)
+		if (FFastGeoPrimitiveComponent* PrimitiveComponent = Component.CastTo<FFastGeoPrimitiveComponent>())
 		{
-			if (FPrimitiveSceneProxy* SceneProxy = Component.GetSceneProxy())
+			FWriteScopeLock WriteLock(*PrimitiveComponent->Lock.Get());
+			const bool bOldIsDrawnInGame = PrimitiveComponent->IsDrawnInGame();
+			PrimitiveComponent->UpdateVisibility();
+			const bool bIsDrawnInGame = PrimitiveComponent->IsDrawnInGame();
+			if (bIsDrawnInGame != bOldIsDrawnInGame)
 			{
-				if (bIsDrawnInGame)
-				{
-					ShowComponents.Add(&Component);
-				}
-				else
+				if (FPrimitiveSceneProxy* SceneProxy = PrimitiveComponent->GetSceneProxy())
 				{
-					HideComponents.Add(&Component);
+					if (bIsDrawnInGame)
+					{
+						ShowComponents.Add(PrimitiveComponent);
+					}
+					else
+					{
+						HideComponents.Add(PrimitiveComponent);
+					}
 				}
 			}
 		}
+		else
+		{
+			Component.UpdateVisibility();
+		}
 	});
 
 	UpdateVisibility_Internal(MoveTemp(ShowComponents), MoveTemp(HideComponents));
@@ -237,6 +250,10 @@ FFastGeoComponent& FFastGeoComponentCluster::AddComponent(FFastGeoElementType In
 	{
 		NewComponent = &ProceduralISMComponents.Emplace_GetRef(ProceduralISMComponents.Num());
 	}
+	else if (InComponentType == FFastGeoDecalComponent::Type)
+	{
+		NewComponent = &DecalComponents.Emplace_GetRef(DecalComponents.Num());
+	}
 	else
 	{
 		unimplemented(); // Should never be reached, missing type handling
@@ -247,5 +264,10 @@ FFastGeoComponent& FFastGeoComponentCluster::AddComponent(FFastGeoElementType In
 
 bool FFastGeoComponentCluster::HasComponents() const
 {
-	return InstancedStaticMeshComponents.Num() > 0 || StaticMeshComponents.Num() > 0 || InstancedSkinnedMeshComponents.Num() > 0 || SkinnedMeshComponents.Num() > 0 || ProceduralISMComponents.Num() > 0;
+	return InstancedStaticMeshComponents.Num() > 0 || 
+		   StaticMeshComponents.Num() > 0 || 
+		   InstancedSkinnedMeshComponents.Num() > 0 || 
+		   SkinnedMeshComponents.Num() > 0 || 
+		   ProceduralISMComponents.Num() > 0 || 
+		   DecalComponents.Num() > 0;
 }
diff --git a/Engine/Plugins/Experimental/FastGeoStreaming/Source/FastGeoStreaming/Private/FastGeoContainer.cpp b/Engine/Plugins/Experimental/FastGeoStreaming/Source/FastGeoStreaming/Private/FastGeoContainer.cpp
--- a/Engine/Plugins/Experimental/FastGeoStreaming/Source/FastGeoStreaming/Private/FastGeoContainer.cpp
+++ b/Engine/Plugins/Experimental/FastGeoStreaming/Source/FastGeoStreaming/Private/FastGeoContainer.cpp
@@ -5,6 +5,7 @@
 #include "FastGeoHLOD.h"
 #include "FastGeoComponent.h"
 #include "FastGeoWorldSubsystem.h"
+#include "FastGeoDestroyRenderStateContext.h"
 #include "FastGeoLog.h"
 #include "AI/Navigation/NavigationElement.h"
 #include "Async/ParallelFor.h"
@@ -185,7 +186,7 @@ void UFastGeoContainer::Register()
 		{
 			ComponentCluster.OnRegister();
 
-			ComponentCluster.ForEachComponent<FFastGeoPrimitiveComponent>([bApplyWorldTransform, &Transform, this](FFastGeoPrimitiveComponent& Component)
+			ComponentCluster.ForEachComponent<FFastGeoComponent>([bApplyWorldTransform, &Transform, this](FFastGeoComponent& Component)
 			{
 				if (bApplyWorldTransform)
 				{
@@ -250,7 +251,7 @@ void UFastGeoContainer::Unregister()
 		{
 			ComponentCluster.OnUnregister();
 
-			ComponentCluster.ForEachComponent<FFastGeoPrimitiveComponent>([this](FFastGeoPrimitiveComponent& Component)
+			ComponentCluster.ForEachComponent<FFastGeoComponent>([this](FFastGeoComponent& Component)
 			{
 				PendingDestroy.RenderState.ComponentsToProcess.Add(&Component);
 			});
@@ -485,7 +486,7 @@ void UFastGeoContainer::OnCreateRenderState_Concurrent()
 					return; 
 				}
 
-				FFastGeoPrimitiveComponent* ComponentToProcess = PendingCreate.RenderState.ComponentsToProcess[PendingCreate.RenderState.TotalNumProcessed + LocalIdx];
+				FFastGeoComponent* ComponentToProcess = PendingCreate.RenderState.ComponentsToProcess[PendingCreate.RenderState.TotalNumProcessed + LocalIdx];
 				ComponentToProcess->CreateRenderState(/*Context = */ nullptr);
 
 				NumProcessed++;
@@ -506,9 +507,9 @@ void UFastGeoContainer::OnDestroyRenderState_Concurrent()
 	// TODO_FASTGEO: Evaluate if we need throttling of the async destruction task
 	check(PendingDestroy.RenderState.NumToProcess == PendingDestroy.RenderState.ComponentsToProcess.Num());
 	{
-		FFastGeoPrimitiveComponent::FFastGeoDestroyRenderStateContext Context(GetWorld()->Scene);
+		FFastGeoDestroyRenderStateContext Context(GetWorld()->Scene);
 
-		for (FFastGeoPrimitiveComponent* Component : PendingDestroy.RenderState.ComponentsToProcess)
+		for (FFastGeoComponent* Component : PendingDestroy.RenderState.ComponentsToProcess)
 		{
 			Component->DestroyRenderState(&Context);
 		}
@@ -721,10 +722,10 @@ void UFastGeoContainer::PrecachePSOs()
 
 	if (IsComponentPSOPrecachingEnabled())
 	{
-		TArray<FFastGeoPrimitiveComponent*> Components;
+		TArray<FFastGeoComponent*> Components;
 		ForEachComponentCluster([&Components](FFastGeoComponentCluster& ComponentCluster)
 		{
-			ComponentCluster.ForEachComponent<FFastGeoPrimitiveComponent>([&Components](FFastGeoPrimitiveComponent& Component)
+			ComponentCluster.ForEachComponent<FFastGeoComponent>([&Components](FFastGeoComponent& Component)
 			{
 				// Mark component so that IsPSOPrecaching() return true even if component's PSO task has not started
 				Component.MarkPrecachePSOsRequired();
@@ -736,7 +737,7 @@ void UFastGeoContainer::PrecachePSOs()
 		PrecachePSOsTask = FastGeo::FAsyncTaskWithAssetRefs::Launch(TEXT("UFastGeoContainer::PrecachePSOs_Task"), Assets, [this, Components = MoveTemp(Components)]()
 		{
 			TRACE_CPUPROFILER_EVENT_SCOPE(UFastGeoContainer::PrecachePSOs_Task);
-			for (FFastGeoPrimitiveComponent* Component : Components)
+			for (FFastGeoComponent* Component : Components)
 			{
 				Component->PrecachePSOs();
 			}
diff --git a/Engine/Plugins/Experimental/FastGeoStreaming/Source/FastGeoStreaming/Private/FastGeoDecalComponent.cpp b/Engine/Plugins/Experimental/FastGeoStreaming/Source/FastGeoStreaming/Private/FastGeoDecalComponent.cpp
new file mode 100644
--- /dev/null
+++ b/Engine/Plugins/Experimental/FastGeoStreaming/Source/FastGeoStreaming/Private/FastGeoDecalComponent.cpp
@@ -0,0 +1,293 @@
+// Copyright Epic Games, Inc. All Rights Reserved.
+
+#include "FastGeoDecalComponent.h"
+#include "FastGeoContainer.h"
+#include "FastGeoWorldSubsystem.h"
+#include "FastGeoRegisteredComponent.h"
+#include "FastGeoDestroyRenderStateContext.h"
+#include "FastGeoPSOPrecacheFinishedTask.h"
+#include "Components/DecalComponent.h"
+#include "SceneProxies/DeferredDecalProxy.h"
+#include "Materials/MaterialInterface.h"
+#include "Engine/World.h"
+#include "TimerManager.h"
+#include "SceneInterface.h"
+
+const FFastGeoElementType FFastGeoDecalComponent::Type(&FFastGeoComponent::Type);
+
+FFastGeoDecalComponent::FFastGeoDecalComponent(int32 InComponentIndex, FFastGeoElementType InType)
+	: Super(InComponentIndex, InType)
+	, DecalSize(FVector::ZeroVector)
+	, bIsVisible(true)
+	, SceneProxy(nullptr)
+	, Lock(MakeUnique<FRWLock>())
+{
+}
+
+FFastGeoDecalComponent::FFastGeoDecalComponent(const FFastGeoDecalComponent& Other)
+	: Super(Other)
+	, LocalTransform(Other.LocalTransform)
+	, WorldTransform(Other.WorldTransform)
+	, DecalSize(Other.DecalSize)
+	, bIsVisible(Other.bIsVisible)
+	, SceneProxyDesc(Other.SceneProxyDesc)
+	, SceneProxy(nullptr)
+	, Lock(MakeUnique<FRWLock>())
+#if !WITH_EDITOR && UE_WITH_PSO_PRECACHING
+	, LatestPSOPrecacheJobSetCompleted(0)
+	, LatestPSOPrecacheJobSet(0)
+#endif
+{
+}
+
+#if WITH_EDITOR
+void FFastGeoDecalComponent::InitializeFromComponent(UActorComponent* Component)
+{
+	Super::InitializeFromComponent(Component);
+
+	UDecalComponent* DecalComponent = CastChecked<UDecalComponent>(Component);
+	DecalComponent->UpdateComponentToWorld();
+	SceneProxyDesc = FDeferredDecalSceneProxyDesc(DecalComponent);
+	LocalTransform = DecalComponent->GetComponentToWorld();
+	WorldTransform = LocalTransform;
+	DecalSize = DecalComponent->DecalSize;
+	bIsVisible = SceneProxyDesc.bDrawInGame;
+	check(SceneProxyDesc.Bounds == CalcBounds());
+}
+#endif
+
+void FFastGeoDecalComponent::Serialize(FArchive& Ar)
+{
+	Super::Serialize(Ar);
+
+	Ar << LocalTransform;
+	Ar << WorldTransform;
+	Ar << DecalSize;
+	FArchive_Serialize_BitfieldBool(Ar, bIsVisible);
+
+	// Serialize persistent data from FDeferredDecalSceneProxyDesc
+	Ar << SceneProxyDesc.DecalMaterial;
+	Ar << SceneProxyDesc.TransformWithDecalScale;
+	Ar << SceneProxyDesc.Bounds;
+	Ar << SceneProxyDesc.DecalColor;
+	Ar << SceneProxyDesc.InitializationWorldTimeSeconds;
+	Ar << SceneProxyDesc.FadeScreenSize;
+	Ar << SceneProxyDesc.FadeDuration;
+	Ar << SceneProxyDesc.FadeStartDelay;
+	Ar << SceneProxyDesc.FadeInDuration;
+	Ar << SceneProxyDesc.FadeInStartDelay;
+	Ar << SceneProxyDesc.SortOrder;
+	FArchive_Serialize_BitfieldBool(Ar, SceneProxyDesc.bDrawInGame);
+	FArchive_Serialize_BitfieldBool(Ar, SceneProxyDesc.bDrawInEditor);
+	FArchive_Serialize_BitfieldBool(Ar, SceneProxyDesc.bShouldFade);
+}
+
+void FFastGeoDecalComponent::ApplyWorldTransform(const FTransform& InTransform)
+{
+	check(!GetOwnerContainer()->IsRegistered());
+	WorldTransform = LocalTransform * InTransform;
+	SceneProxyDesc.Bounds = CalcBounds();
+	SceneProxyDesc.TransformWithDecalScale = GetTransformIncludingDecalSize();
+}
+
+void FFastGeoDecalComponent::UpdateVisibility()
+{
+	UpdateVisibilityInternal();
+	MarkRenderStateDirty();
+}
+
+void FFastGeoDecalComponent::UpdateVisibilityInternal()
+{
+	SceneProxyDesc.bDrawInGame = bIsVisible && GetOwnerComponentCluster()->IsVisible();
+#if !UE_BUILD_SHIPPING
+	SceneProxyDesc.bDrawInGame = SceneProxyDesc.bDrawInGame && UFastGeoWorldSubsystem::IsFastGeoVisible();
+#endif
+}
+
+bool FFastGeoDecalComponent::ShouldCreateRenderState() const
+{
+	return FApp::CanEverRender() && SceneProxyDesc.bDrawInGame;
+}
+
+void FFastGeoDecalComponent::CreateRenderState(FRegisterComponentContext* Context)
+{
+	FWriteScopeLock WriteLock(*Lock.Get());
+	check(ProxyState == EProxyCreationState::None);
+	check(!SceneProxy);
+	ProxyState = EProxyCreationState::Creating;
+	bRenderStateDirty = false;
+
+	FSceneInterface* Scene = GetScene();
+	check(Scene);
+
+	ESceneProxyCreationError Error;
+	SceneProxy = CreateSceneProxy(Error);
+	if (SceneProxy)
+	{
+		const float SpawnTime = GetWorld()->GetTimeSeconds();
+		TArray<FDeferredDecalUpdateParams> DecalsToAdd;
+		DecalsToAdd.Add(
+		{
+			.Transform = SceneProxy->ComponentTrans,
+			.DecalProxy = SceneProxy,
+			.Bounds = SceneProxy->GetBounds(),
+			.AbsSpawnTime = SpawnTime,
+			.FadeDuration = SceneProxyDesc.FadeDuration,
+			.FadeStartDelay = SceneProxyDesc.FadeStartDelay,
+			.FadeInDuration = SceneProxyDesc.FadeInDuration,
+			.FadeInStartDelay = SceneProxyDesc.FadeInStartDelay,
+			.FadeScreenSize = SceneProxyDesc.FadeScreenSize,
+			.SortOrder = SceneProxyDesc.SortOrder,
+			.DecalColor = SceneProxyDesc.DecalColor,
+			.OperationType = FDeferredDecalUpdateParams::EOperationType::AddToSceneAndUpdate,
+		});
+		Scene->BatchUpdateDecals(MoveTemp(DecalsToAdd));
+
+		ProxyState = EProxyCreationState::Created;
+
+		const float LifeSpan = SceneProxyDesc.FadeStartDelay + SceneProxyDesc.FadeDuration;
+		if (LifeSpan > 0)
+		{
+			AsyncTask(ENamedThreads::GameThread, [Component = FFastGeoRegisteredComponent(this), LifeSpan]()
+			{
+				if (FFastGeoDecalComponent* DecalComponent = Component.TryGetRegistered<FFastGeoDecalComponent>())
+				{
+					DecalComponent->SetLifeSpan(LifeSpan);
+				}
+			});
+		}
+	}
+	else if (Error == ESceneProxyCreationError::WaitingPSOs)
+	{
+		ProxyState = EProxyCreationState::Delayed;
+	}
+	else
+	{
+		ProxyState = EProxyCreationState::None;
+	}
+}
+
+FDeferredDecalProxy* FFastGeoDecalComponent::CreateSceneProxy(ESceneProxyCreationError& OutError)
+{
+	check(!SceneProxy);
+
+#if !WITH_EDITOR && UE_WITH_PSO_PRECACHING
+	if (LatestPSOPrecacheJobSetCompleted != LatestPSOPrecacheJobSet && GetPSOPrecacheProxyCreationStrategy() == EPSOPrecacheProxyCreationStrategy::DelayUntilPSOPrecached)
+	{
+		OutError = ESceneProxyCreationError::WaitingPSOs;
+		return nullptr;
+	}
+#endif
+
+	OutError = ESceneProxyCreationError::None;
+	InitializeSceneProxyDescDynamicProperties();
+	SceneProxy = new FDeferredDecalProxy(SceneProxyDesc);
+	return SceneProxy;
+}
+
+void FFastGeoDecalComponent::InitializeSceneProxyDescDynamicProperties()
+{
+	UpdateVisibilityInternal();
+}
+
+void FFastGeoDecalComponent::DestroyRenderState(FFastGeoDestroyRenderStateContext* Context)
+{
+	FWriteScopeLock WriteLock(*Lock.Get());
+	if (SceneProxy)
+	{
+		FSceneInterface* Scene = GetScene();
+		check(Scene);
+		check(ProxyState == EProxyCreationState::Created);
+
+		FFastGeoDestroyRenderStateContext::DestroyProxy(Context, SceneProxy, Scene);
+		
+		SceneProxy = nullptr;
+		ProxyState = EProxyCreationState::None;
+		bRenderStateDirty = false;
+	}
+	else
+	{
+		check(ProxyState == EProxyCreationState::None || ProxyState == EProxyCreationState::Delayed);
+		ProxyState = EProxyCreationState::None;
+	}
+}
+
+FBoxSphereBounds FFastGeoDecalComponent::CalcBounds() const
+{
+	return FBoxSphereBounds(FVector(0, 0, 0), DecalSize, DecalSize.Size()).TransformBy(WorldTransform);
+}
+
+FTransform FFastGeoDecalComponent::GetTransformIncludingDecalSize() const
+{
+	FTransform Ret = WorldTransform;
+	Ret.SetScale3D(Ret.GetScale3D() * DecalSize);
+	return Ret;
+}
+
+void FFastGeoDecalComponent::SetLifeSpan(const float InLifeSpan)
+{
+	check(IsInGameThread());
+	if (TimerHandle_DestroyDecalComponent.IsValid())
+	{
+		GetWorld()->GetTimerManager().ClearTimer(TimerHandle_DestroyDecalComponent);
+	}
+	if (InLifeSpan > 0.f)
+	{
+		FFastGeoRegisteredComponent Component(this);
+		GetWorld()->GetTimerManager().SetTimer(TimerHandle_DestroyDecalComponent, FTimerDelegate::CreateLambda([Component]()
+		{
+			if (FFastGeoDecalComponent* DecalComponent = Component.TryGetRegistered<FFastGeoDecalComponent>())
+			{
+				DecalComponent->DestroyRenderState(nullptr);
+			}
+		}), InLifeSpan, false);
+	}
+}
+
+void FFastGeoDecalComponent::PrecachePSOs()
+{
+#if !WITH_EDITOR && UE_WITH_PSO_PRECACHING
+	TRACE_CPUPROFILER_EVENT_SCOPE(FFastGeoPrimitiveComponent::PrecachePSOs);
+
+	if (!FApp::CanEverRender() || !IsComponentPSOPrecachingEnabled())
+	{
+		return;
+	}
+
+	if (SceneProxyDesc.DecalMaterial && !SceneProxyDesc.DecalMaterial->HasAnyFlags(RF_NeedPostLoad))
+	{
+		FPSOPrecacheParams PSOPrecacheParams;
+		FPSOPrecacheVertexFactoryDataList VertexFactoryDataList;
+		VertexFactoryDataList.Add(FPSOPrecacheVertexFactoryData(&FLocalVertexFactory::StaticType));
+
+		// Immediately create at high priority and thus doesn't need boosting anymore
+		TArray<FMaterialPSOPrecacheRequestID> MaterialPSOPrecacheRequestIDs;
+		FGraphEventArray GraphEvents = SceneProxyDesc.DecalMaterial->PrecachePSOs(VertexFactoryDataList, PSOPrecacheParams, EPSOPrecachePriority::High, MaterialPSOPrecacheRequestIDs);
+
+		// Request recreate of the render state when the PSO compilation is ready (if we want to delay proxy creation)
+		if (GetPSOPrecacheProxyCreationStrategy() != EPSOPrecacheProxyCreationStrategy::AlwaysCreate)
+		{
+			LatestPSOPrecacheJobSet++;
+			if (GraphEvents.Num() > 0)
+			{
+				TGraphTask<FFastGeoPSOPrecacheFinishedTask>::CreateTask(&GraphEvents).ConstructAndDispatchWhenReady(this, LatestPSOPrecacheJobSet);
+			}
+			else
+			{
+				// No graph events to wait on, the job set can be considered complete.
+				LatestPSOPrecacheJobSetCompleted = LatestPSOPrecacheJobSet;
+			}
+		}
+	}
+#endif
+}
+
+#if !WITH_EDITOR && UE_WITH_PSO_PRECACHING
+void FFastGeoDecalComponent::OnPrecacheFinished(int32 JobSetThatJustCompleted)
+{
+	QUICK_SCOPE_CYCLE_COUNTER(STAT_PSOPrecacheFinishedTask);
+	int32 CurrJobSetCompleted = LatestPSOPrecacheJobSetCompleted.load();
+	while (CurrJobSetCompleted < JobSetThatJustCompleted && !LatestPSOPrecacheJobSetCompleted.compare_exchange_weak(CurrJobSetCompleted, JobSetThatJustCompleted)) {}
+	MarkRenderStateDirty();
+}
+#endif
\ No newline at end of file
diff --git a/Engine/Plugins/Experimental/FastGeoStreaming/Source/FastGeoStreaming/Private/FastGeoDestroyRenderStateContext.cpp b/Engine/Plugins/Experimental/FastGeoStreaming/Source/FastGeoStreaming/Private/FastGeoDestroyRenderStateContext.cpp
new file mode 100644
--- /dev/null
+++ b/Engine/Plugins/Experimental/FastGeoStreaming/Source/FastGeoStreaming/Private/FastGeoDestroyRenderStateContext.cpp
@@ -0,0 +1,69 @@
+// Copyright Epic Games, Inc. All Rights Reserved.
+
+#include "FastGeoDestroyRenderStateContext.h"
+#include "SceneProxies/DeferredDecalProxy.h"
+#include "PrimitiveSceneProxy.h"
+#include "SceneInterface.h"
+
+FFastGeoDestroyRenderStateContext::FFastGeoDestroyRenderStateContext(FSceneInterface* InScene)
+	: Scene(InScene)
+{
+}
+
+FFastGeoDestroyRenderStateContext::~FFastGeoDestroyRenderStateContext()
+{
+	if (!PrimitiveSceneProxies.IsEmpty())
+	{
+		Scene->BatchRemovePrimitives(MoveTemp(PrimitiveSceneProxies));
+	}
+	if (!DeferredDecalProxies.IsEmpty())
+	{
+		TArray<FDeferredDecalUpdateParams> BatchParams;
+		BatchParams.Reserve(DeferredDecalProxies.Num());
+		for (FDeferredDecalProxy* Proxy : DeferredDecalProxies)
+		{
+			BatchParams.Add(
+			{
+				.DecalProxy = Proxy,
+				.OperationType = FDeferredDecalUpdateParams::EOperationType::RemoveFromSceneAndDelete
+			});
+		}
+		Scene->BatchUpdateDecals(MoveTemp(BatchParams));
+	}
+}
+
+void FFastGeoDestroyRenderStateContext::DestroyProxy(FFastGeoDestroyRenderStateContext* InContext, FPrimitiveSceneProxy* InPrimitiveSceneProxy, FSceneInterface* InScene)
+{
+	check(InScene);
+	check(InPrimitiveSceneProxy);
+
+	if (InContext)
+	{
+		InContext->PrimitiveSceneProxies.Add(InPrimitiveSceneProxy);
+	}
+	else
+	{
+		InScene->BatchRemovePrimitives({ InPrimitiveSceneProxy });
+	}
+}
+
+void FFastGeoDestroyRenderStateContext::DestroyProxy(FFastGeoDestroyRenderStateContext* InContext, FDeferredDecalProxy* InDeferredDecalProxy, FSceneInterface* InScene)
+{
+	check(InScene);
+	check(InDeferredDecalProxy);
+
+	if (InContext)
+	{
+		InContext->DeferredDecalProxies.Add(InDeferredDecalProxy);
+	}
+	else
+	{
+		TArray<FDeferredDecalUpdateParams> BatchParams;
+		BatchParams.Add(
+		{
+			.DecalProxy = InDeferredDecalProxy,
+			.OperationType = FDeferredDecalUpdateParams::EOperationType::RemoveFromSceneAndDelete
+		});
+		InScene->BatchUpdateDecals(MoveTemp(BatchParams));
+	}
+}
\ No newline at end of file
diff --git a/Engine/Plugins/Experimental/FastGeoStreaming/Source/FastGeoStreaming/Private/FastGeoPrimitiveComponent.cpp b/Engine/Plugins/Experimental/FastGeoStreaming/Source/FastGeoStreaming/Private/FastGeoPrimitiveComponent.cpp
--- a/Engine/Plugins/Experimental/FastGeoStreaming/Source/FastGeoStreaming/Private/FastGeoPrimitiveComponent.cpp
+++ b/Engine/Plugins/Experimental/FastGeoStreaming/Source/FastGeoStreaming/Private/FastGeoPrimitiveComponent.cpp
@@ -6,8 +6,10 @@
 #include "FastGeoInstancedStaticMeshComponent.h"
 #include "FastGeoWorldSubsystem.h"
 #include "FastGeoComponentCluster.h"
+#include "FastGeoDestroyRenderStateContext.h"
 #include "FastGeoWeakElement.h"
 #include "FastGeoLog.h"
+#include "FastGeoPSOPrecacheFinishedTask.h"
 #include "AI/NavigationModifier.h"
 #include "AI/Navigation/NavigationRelevantData.h"
 #include "PrimitiveSceneDesc.h"
@@ -27,38 +29,6 @@
 #include "ObjectCacheEventSink.h"
 #endif
 
-#if !UE_BUILD_SHIPPING
-namespace FastGeo
-{
-	static int32 GShowFastGeo = 1;
-	FAutoConsoleCommand GShowFastGeoCommand(
-		TEXT("FastGeo.Show"),
-		TEXT("Turn on/off rendering of FastGeo."),
-		FConsoleCommandWithArgsDelegate::CreateLambda([](const TArray<FString>& Args)
-		{
-			GShowFastGeo = (Args.Num() != 1) || (Args[0] != TEXT("0"));
-			for (const FWorldContext& Context : GEngine->GetWorldContexts())
-			{
-				UWorld* World = Context.World();
-				if (World && World->IsGameWorld())
-				{
-					for (ULevel* Level : World->GetLevels())
-					{
-						if (UFastGeoContainer* FastGeo = Level->GetAssetUserData<UFastGeoContainer>())
-						{
-							FastGeo->ForEachComponentCluster([](FFastGeoComponentCluster& ComponentCluster)
-							{
-								ComponentCluster.UpdateVisibility();
-							});
-						}
-					}
-				}
-			}
-		})
-	);
-}
-#endif
-
 const FFastGeoElementType FFastGeoPrimitiveComponent::Type(&FFastGeoComponent::Type);
 
 FFastGeoPrimitiveComponent::FFastGeoPrimitiveComponent(int32 InComponentIndex, FFastGeoElementType InType)
@@ -88,8 +58,6 @@ FFastGeoPrimitiveComponent::FFastGeoPrimitiveComponent(const FFastGeoPrimitiveCo
 	, BodyInstanceOwner()
 	, PrimitiveSceneData()
 	, AsyncTermBodyPayload()
-	, ProxyState(EProxyCreationState::None)
-	, bRenderStateDirty(false)
 #if !WITH_EDITOR && UE_WITH_PSO_PRECACHING
 	, MaterialPSOPrecacheRequestIDs()
 	, LatestPSOPrecacheJobSetCompleted(0)
@@ -286,7 +254,7 @@ void FFastGeoPrimitiveComponent::UpdateVisibility()
 	FPrimitiveSceneProxyDesc& SceneProxyDesc = GetSceneProxyDesc();
 	SceneProxyDesc.bIsVisible = bIsVisible && GetOwnerComponentCluster()->IsVisible();
 #if !UE_BUILD_SHIPPING
-	SceneProxyDesc.bIsVisible = SceneProxyDesc.bIsVisible && FastGeo::GShowFastGeo;
+	SceneProxyDesc.bIsVisible = SceneProxyDesc.bIsVisible && UFastGeoWorldSubsystem::IsFastGeoVisible();
 #endif
 }
 
@@ -377,44 +345,17 @@ void FFastGeoPrimitiveComponent::CreateRenderState(FRegisterComponentContext* Co
 	}
 }
 
-FFastGeoPrimitiveComponent::FFastGeoDestroyRenderStateContext::FFastGeoDestroyRenderStateContext(FSceneInterface* InScene)
-	: Scene(InScene)
-{
-}
-
-FFastGeoPrimitiveComponent::FFastGeoDestroyRenderStateContext::~FFastGeoDestroyRenderStateContext()
-{
-	if (HasPendingWork())
-	{
-		Scene->BatchRemovePrimitives(MoveTemp(PrimitiveSceneProxies));
-	}
-}
-
-bool FFastGeoPrimitiveComponent::FFastGeoDestroyRenderStateContext::HasPendingWork() const
-{
-	return !PrimitiveSceneProxies.IsEmpty();
-}
-
-void FFastGeoPrimitiveComponent::FFastGeoDestroyRenderStateContext::DestroyProxy(FFastGeoDestroyRenderStateContext* InContext, FPrimitiveSceneProxy* InPrimitiveSceneProxy)
-{
-	if (InContext)
-	{
-		InContext->PrimitiveSceneProxies.Add(InPrimitiveSceneProxy);
-	}
-	else
-	{
-		InPrimitiveSceneProxy->GetScene().BatchRemovePrimitives({ InPrimitiveSceneProxy });
-	}
-}
-
 void FFastGeoPrimitiveComponent::DestroyRenderState(FFastGeoDestroyRenderStateContext* Context)
 {
 	FWriteScopeLock WriteLock(*Lock.Get());
 	if (GetSceneProxy())
 	{
+		FSceneInterface* Scene = GetScene();
+		check(Scene);
+
 		check(ProxyState == EProxyCreationState::Created);
 
-		FFastGeoDestroyRenderStateContext::DestroyProxy(Context, GetSceneProxy());
+		FFastGeoDestroyRenderStateContext::DestroyProxy(Context, GetSceneProxy(), Scene);
 
 		PrimitiveSceneData.SceneProxy = nullptr;
 
@@ -532,21 +473,6 @@ void FFastGeoPrimitiveComponent::OnAsyncDestroyPhysicsState()
 	Super::OnAsyncDestroyPhysicsState();
 }
 
-bool FFastGeoPrimitiveComponent::IsRenderStateCreated() const
-{
-	return ProxyState == EProxyCreationState::Created;
-}
-
-bool FFastGeoPrimitiveComponent::IsRenderStateDelayed() const
-{
-	return ProxyState == EProxyCreationState::Delayed;
-}
-
-bool FFastGeoPrimitiveComponent::IsRenderStateDirty() const
-{
-	return bRenderStateDirty;
-}
-
 bool FFastGeoPrimitiveComponent::ShouldCreateRenderState() const
 {
 	if (!FApp::CanEverRender())
@@ -576,35 +502,11 @@ bool FFastGeoPrimitiveComponent::IsCollisionEnabled() const
 	return GetSceneProxyDesc().bCollisionEnabled;
 }
 
-FSceneInterface* FFastGeoPrimitiveComponent::GetScene() const
-{
-	if (UWorld* World = GetWorld())
-	{
-		return World->Scene;
-	}
-	return nullptr;
-}
-
 FPrimitiveSceneProxy* FFastGeoPrimitiveComponent::GetSceneProxy() const
 {
 	return PrimitiveSceneData.SceneProxy;
 }
 
-void FFastGeoPrimitiveComponent::MarkRenderStateDirty()
-{
-	FReadScopeLock ReadLock(*Lock.Get());
-	if ((IsRenderStateCreated() || IsRenderStateDelayed()) && !IsRenderStateDirty())
-	{
-		bRenderStateDirty = true;
-
-		UFastGeoWorldSubsystem* WorldSubsystem = GetWorld()->GetSubsystem<UFastGeoWorldSubsystem>();
-		if (ensure(WorldSubsystem))
-		{
-			WorldSubsystem->AddToComponentsPendingRecreate(this);
-		}
-	}
-}
-
 const FTransform& FFastGeoPrimitiveComponent::GetTransform() const
 {
 	return WorldTransform;
@@ -641,34 +543,6 @@ void FFastGeoPrimitiveComponent::InitializeDynamicProperties()
 }
 
 #if !WITH_EDITOR && UE_WITH_PSO_PRECACHING
-struct FFastGeoPSOPrecacheFinishedTask
-{
-	explicit FFastGeoPSOPrecacheFinishedTask(FFastGeoPrimitiveComponent* InPrimitiveComponent, int32 InJobSetThatJustCompleted)
-		: WeakPrimitiveComponent(InPrimitiveComponent)
-		, JobSetThatJustCompleted(InJobSetThatJustCompleted)
-	{
-	}
-
-	static TStatId GetStatId() { return TStatId(); }
-	static ENamedThreads::Type GetDesiredThread() { return ENamedThreads::GameThread; }
-	static ESubsequentsMode::Type GetSubsequentsMode() { return ESubsequentsMode::TrackSubsequents; }
-
-	void DoTask(ENamedThreads::Type CurrentThread, const FGraphEventRef& MyCompletionGraphEvent)
-	{
-		if (FFastGeoPrimitiveComponent* PrimitiveComponent = WeakPrimitiveComponent.Get<FFastGeoPrimitiveComponent>())
-		{
-			// Validate that the component is still part of a streamed-in level
-			if (PrimitiveComponent->GetWorld())
-			{
-				PrimitiveComponent->OnPrecacheFinished(JobSetThatJustCompleted);
-			}
-		}
-	}
-
-	FWeakFastGeoComponent WeakPrimitiveComponent;
-	int32 JobSetThatJustCompleted;
-};
-
 void FFastGeoPrimitiveComponent::OnPrecacheFinished(int32 JobSetThatJustCompleted)
 {
 	QUICK_SCOPE_CYCLE_COUNTER(STAT_PSOPrecacheFinishedTask);
diff --git a/Engine/Plugins/Experimental/FastGeoStreaming/Source/FastGeoStreaming/Private/FastGeoRegisteredComponent.h b/Engine/Plugins/Experimental/FastGeoStreaming/Source/FastGeoStreaming/Private/FastGeoRegisteredComponent.h
new file mode 100644
--- /dev/null
+++ b/Engine/Plugins/Experimental/FastGeoStreaming/Source/FastGeoStreaming/Private/FastGeoRegisteredComponent.h
@@ -0,0 +1,41 @@
+// Copyright Epic Games, Inc. All Rights Reserved.
+
+#pragma once
+
+#include "FastGeoComponent.h"
+#include "FastGeoContainer.h"
+#include "FastGeoWeakElement.h"
+
+struct FFastGeoRegisteredComponent
+{
+	FFastGeoRegisteredComponent(FFastGeoComponent* InComponent)
+		: Component(InComponent)
+	{
+		check(InComponent);
+		check(InComponent->GetOwnerContainer());
+		RegistrationEpoch = InComponent->GetOwnerContainer()->RegistrationEpoch;
+	}
+
+	template <class T>
+	FORCEINLINE T* TryGetRegistered() const
+	{
+		if (T* ComponentPtr = Component.Get<T>())
+		{
+			if (ComponentPtr->GetOwnerContainer()->IsRegistered() &&
+				RegistrationEpoch == ComponentPtr->GetOwnerContainer()->RegistrationEpoch)
+			{
+				return ComponentPtr;
+			}
+		}
+		return nullptr;
+	}
+
+	FORCEINLINE bool operator==(const FFastGeoRegisteredComponent& Other) const
+	{
+		return Component == Other.Component && RegistrationEpoch == Other.RegistrationEpoch;
+	}
+
+private:
+	FWeakFastGeoComponent Component;
+	int32 RegistrationEpoch = -1;
+};
\ No newline at end of file
diff --git a/Engine/Plugins/Experimental/FastGeoStreaming/Source/FastGeoStreaming/Private/FastGeoWorldPartitionRuntimeCellTransformer.cpp b/Engine/Plugins/Experimental/FastGeoStreaming/Source/FastGeoStreaming/Private/FastGeoWorldPartitionRuntimeCellTransformer.cpp
--- a/Engine/Plugins/Experimental/FastGeoStreaming/Source/FastGeoStreaming/Private/FastGeoWorldPartitionRuntimeCellTransformer.cpp
+++ b/Engine/Plugins/Experimental/FastGeoStreaming/Source/FastGeoStreaming/Private/FastGeoWorldPartitionRuntimeCellTransformer.cpp
@@ -20,6 +20,7 @@
 #include "Components/SkinnedMeshComponent.h"
 #include "Components/InstancedSkinnedMeshComponent.h"
 #include "Components/SplineMeshComponent.h"
+#include "Components/DecalComponent.h"
 #include "GameFramework/WorldSettings.h"
 #include "GameFramework/Info.h"
 #include "WorldPartition/WorldPartition.h"
@@ -102,14 +103,15 @@ namespace FastGeo
 		}
 	};
 
-	FFastGeoElementType GetFastGeoComponentType(TSubclassOf<UPrimitiveComponent> InClass)
+	FFastGeoElementType GetFastGeoComponentType(TSubclassOf<USceneComponent> InClass)
 	{
-		static const TMap<TSubclassOf<UPrimitiveComponent>, FFastGeoElementType> FastGeoComponentTypeMapping =
+		static const TMap<TSubclassOf<USceneComponent>, FFastGeoElementType> FastGeoComponentTypeMapping =
 		{
 			{ UStaticMeshComponent::StaticClass(), FFastGeoStaticMeshComponent::Type },
 			{ UInstancedStaticMeshComponent::StaticClass(), FFastGeoInstancedStaticMeshComponent::Type },
 			{ USkinnedMeshComponent::StaticClass(), FFastGeoSkinnedMeshComponent::Type },
-			{ UInstancedSkinnedMeshComponent::StaticClass(), FFastGeoInstancedSkinnedMeshComponent::Type }
+			{ UInstancedSkinnedMeshComponent::StaticClass(), FFastGeoInstancedSkinnedMeshComponent::Type },
+			{ UDecalComponent::StaticClass(), FFastGeoDecalComponent::Type }
 		};
 
 		// Walk the component class hierarchy and look for a fast geo mapping
@@ -763,7 +765,9 @@ FFastGeoTransformResult UFastGeoWorldPartitionRuntimeCellTransformer::CanTransfo
 		return { EFastGeoTransform::Reject, [InActor] { return FString::Printf(TEXT("Actor %s is replicated"), *InActor->GetName()); } };
 	}
 
-	if (!InActor->IsRootComponentStatic())
+	// Ignore DecalComponent for now, as its mobility is Movable
+	const bool bIsRootDecalComponent = InActor->GetRootComponent() && InActor->GetRootComponent()->IsA<UDecalComponent>();
+	if (!InActor->IsRootComponentStatic() && !bIsRootDecalComponent)
 	{
 		return { EFastGeoTransform::Reject, [InActor] { return FString::Printf(TEXT("Actor %s RootComponent Mobility is not Static"), *InActor->GetName()); } };
 	}
@@ -887,45 +891,49 @@ FFastGeoTransformResult UFastGeoWorldPartitionRuntimeCellTransformer::IsAllowedC
 
 FFastGeoTransformResult UFastGeoWorldPartitionRuntimeCellTransformer::CanTransformComponent(UActorComponent* InComponent) const
 {
-	UPrimitiveComponent* PrimitiveComponent = Cast<UPrimitiveComponent>(InComponent);
-	if (!PrimitiveComponent)
+	if (InComponent->IsEditorOnly())
+	{
+		return { EFastGeoTransform::Discard, [InComponent] { return FString::Printf(TEXT("Component %s is editor-only"), *FastGeo::GetComponentShortName(InComponent)); } };
+	}
+
+	if (CanIgnoreComponent(InComponent))
+	{
+		return { EFastGeoTransform::Discard };
+	}
+
+	USceneComponent* SceneComponent = Cast<USceneComponent>(InComponent);
+	if (!SceneComponent)
 	{
-		if (CanIgnoreComponent(InComponent))
-		{
-			return { EFastGeoTransform::Discard };
-		}
 		return { EFastGeoTransform::Reject, [InComponent] { return FString::Printf(TEXT("Component %s class %s is not allowed"), *FastGeo::GetComponentShortName(InComponent), *InComponent->GetClass()->GetName()); } }; //-V522
 	}
 
-	FFastGeoTransformResult AllowedComponentClassResult = IsAllowedComponentClass(InComponent);
+	const FFastGeoTransformResult AllowedComponentClassResult = IsAllowedComponentClass(InComponent);
 	if (AllowedComponentClassResult.GetResult() != EFastGeoTransform::Allow)
 	{
 		return AllowedComponentClassResult;
 	}
 
 	FString Reason;
-	if (!IsComponentTransformable(PrimitiveComponent, Reason))
+	if (!IsComponentTransformable(SceneComponent, Reason))
 	{
-		return { EFastGeoTransform::Reject, [PrimitiveComponent, Reason] { return FString::Printf(TEXT("Component %s [%s]"), *FastGeo::GetComponentShortName(PrimitiveComponent), *Reason); } };
+		return { EFastGeoTransform::Reject, [SceneComponent, Reason] { return FString::Printf(TEXT("Component %s [%s]"), *FastGeo::GetComponentShortName(SceneComponent), *Reason); } };
 	}
 
-	if (PrimitiveComponent->IsEditorOnly())
+	// Ignore DecalComponent for now, as its mobility is Movable
+	if (SceneComponent->Mobility != EComponentMobility::Static && !SceneComponent->IsA<UDecalComponent>())
 	{
-		return { EFastGeoTransform::Discard, [PrimitiveComponent] { return FString::Printf(TEXT("Component %s is editor-only"), *FastGeo::GetComponentShortName(PrimitiveComponent)); } };
+		return { EFastGeoTransform::Reject, [SceneComponent] { return FString::Printf(TEXT("Component %s Mobility is not Static"), *FastGeo::GetComponentShortName(SceneComponent)); } };
 	}
 
-	if (PrimitiveComponent->GetLODParentPrimitive())
+	UPrimitiveComponent* PrimitiveComponent = Cast<UPrimitiveComponent>(InComponent);
+	if (PrimitiveComponent && PrimitiveComponent->GetLODParentPrimitive())
 	{
 		return { EFastGeoTransform::Reject, [PrimitiveComponent] { return FString::Printf(TEXT("Component %s has a valid LOD Parent Primitive"), *FastGeo::GetComponentShortName(PrimitiveComponent)); } };
 	}
 
-	if (PrimitiveComponent->Mobility != EComponentMobility::Static)
-	{
-		return { EFastGeoTransform::Reject, [PrimitiveComponent] { return FString::Printf(TEXT("Component %s Mobility is not Static"), *FastGeo::GetComponentShortName(PrimitiveComponent)); } };
-	}
-
-	bool bIsComponentVisible = PrimitiveComponent->IsVisible() && !PrimitiveComponent->GetOwner()->IsHidden();
-	bool bShouldAddToRenderScene = bIsComponentVisible || PrimitiveComponent->bCastHiddenShadow || PrimitiveComponent->bAffectIndirectLightingWhileHidden || PrimitiveComponent->bRayTracingFarField;
+	bool bIsSceneComponentVisible = SceneComponent->IsVisible() && !SceneComponent->GetOwner()->IsHidden();
+	bool bIsPrimitiveVisible = PrimitiveComponent && (PrimitiveComponent->bCastHiddenShadow || PrimitiveComponent->bAffectIndirectLightingWhileHidden || PrimitiveComponent->bRayTracingFarField);
+	bool bShouldAddToRenderScene = bIsSceneComponentVisible || bIsPrimitiveVisible;
 
 	if (UStaticMeshComponent* StaticMeshComponent = Cast<UStaticMeshComponent>(PrimitiveComponent))
 	{
@@ -1009,18 +1017,21 @@ FFastGeoTransformResult UFastGeoWorldPartitionRuntimeCellTransformer::CanTransfo
 		}
 	}
 
-	const bool bIsCollisionEnabled = FastGeo::IsCollisionEnabled(PrimitiveComponent);
-
-	// If collision is enabled, only allow if async physics state creation and destruction are supported 
-	if (bIsCollisionEnabled && (!PrimitiveComponent->AllowsAsyncPhysicsStateCreation() || !PrimitiveComponent->AllowsAsyncPhysicsStateDestruction()))
+	if (PrimitiveComponent)
 	{
-		return { EFastGeoTransform::Reject, [PrimitiveComponent] { return FString::Printf(TEXT("Component %s has collision enabled but doesn't allow asynchronous physics state creation/destruction"), *FastGeo::GetComponentShortName(PrimitiveComponent)); } };
-	}
+		const bool bIsCollisionEnabled = FastGeo::IsCollisionEnabled(PrimitiveComponent);
 
-	// Disallow transform if collision is disabled and component doesn't need to be added to the render scene
-	if (!bIsCollisionEnabled && !bShouldAddToRenderScene)
-	{
-		return { EFastGeoTransform::Discard, [PrimitiveComponent] { return FString::Printf(TEXT("Component %s has no collision and is not visible"), *FastGeo::GetComponentShortName(PrimitiveComponent)); } };
+		// If collision is enabled, only allow if async physics state creation and destruction are supported 
+		if (bIsCollisionEnabled && (!PrimitiveComponent->AllowsAsyncPhysicsStateCreation() || !PrimitiveComponent->AllowsAsyncPhysicsStateDestruction()))
+		{
+			return { EFastGeoTransform::Reject, [PrimitiveComponent] { return FString::Printf(TEXT("Component %s has collision enabled but doesn't allow asynchronous physics state creation/destruction"), *FastGeo::GetComponentShortName(PrimitiveComponent)); } };
+		}
+
+		// Disallow transform if collision is disabled and component doesn't need to be added to the render scene
+		if (!bIsCollisionEnabled && !bShouldAddToRenderScene)
+		{
+			return { EFastGeoTransform::Discard, [PrimitiveComponent] { return FString::Printf(TEXT("Component %s has no collision and is not visible"), *FastGeo::GetComponentShortName(PrimitiveComponent)); } };
+		}
 	}
 
 	if (CanIgnoreComponent(InComponent))
diff --git a/Engine/Plugins/Experimental/FastGeoStreaming/Source/FastGeoStreaming/Private/FastGeoWorldSubsystem.cpp b/Engine/Plugins/Experimental/FastGeoStreaming/Source/FastGeoStreaming/Private/FastGeoWorldSubsystem.cpp
--- a/Engine/Plugins/Experimental/FastGeoStreaming/Source/FastGeoStreaming/Private/FastGeoWorldSubsystem.cpp
+++ b/Engine/Plugins/Experimental/FastGeoStreaming/Source/FastGeoStreaming/Private/FastGeoWorldSubsystem.cpp
@@ -35,6 +35,35 @@ namespace FastGeo
 		GAsyncRenderStateTaskMaxNumComponentsToProcess,
 		TEXT("Maximum number of components to process (0 = no component limit)"),
 		ECVF_Default);
+
+#if !UE_BUILD_SHIPPING
+	static bool GShowFastGeo = true;
+	FAutoConsoleCommand GShowFastGeoCommand(
+		TEXT("FastGeo.Show"),
+		TEXT("Turn on/off rendering of FastGeo."),
+		FConsoleCommandWithArgsDelegate::CreateLambda([](const TArray<FString>& Args)
+		{
+			GShowFastGeo = (Args.Num() != 1) || (Args[0] != TEXT("0"));
+			for (const FWorldContext& Context : GEngine->GetWorldContexts())
+			{
+				UWorld* World = Context.World();
+				if (World && World->IsGameWorld())
+				{
+					for (ULevel* Level : World->GetLevels())
+					{
+						if (UFastGeoContainer* FastGeo = Level->GetAssetUserData<UFastGeoContainer>())
+						{
+							FastGeo->ForEachComponentCluster([](FFastGeoComponentCluster& ComponentCluster)
+							{
+								ComponentCluster.UpdateVisibility();
+							});
+						}
+					}
+				}
+			}
+		})
+	);
+#endif
 }
 
 CSV_DEFINE_CATEGORY(FastGeo, true);
@@ -44,6 +73,13 @@ bool UFastGeoWorldSubsystem::IsEnableDebugView()
 	return bEnableDebugView;
 }
 
+#if !UE_BUILD_SHIPPING
+bool UFastGeoWorldSubsystem::IsFastGeoVisible()
+{
+	return FastGeo::GShowFastGeo;
+}
+#endif
+
 UFastGeoWorldSubsystem::UFastGeoWorldSubsystem()
 {
 #if ENABLE_ACTOR_PRIMITIVE_COLOR_HANDLER
@@ -314,7 +350,7 @@ void UFastGeoWorldSubsystem::OnUpdateLevelStreaming()
 	}
 }
 
-void UFastGeoWorldSubsystem::AddToComponentsPendingRecreate(FFastGeoPrimitiveComponent* InComponentPendingRecreate)
+void UFastGeoWorldSubsystem::AddToComponentsPendingRecreate(FFastGeoComponent* InComponentPendingRecreate)
 {
 	ComponentsPendingRecreate.Emplace(InComponentPendingRecreate);
 }
@@ -326,10 +362,10 @@ void UFastGeoWorldSubsystem::ProcessPendingRecreate()
 	CSV_CUSTOM_STAT(FastGeo, PendingRecreate, ComponentsPendingRecreate.Num(), ECsvCustomStatOp::Set);
 
 	int32 NumPrimitiveDelayed = 0;
-	for (const FRegisteredComponent& PendingRecreateComponent : ComponentsPendingRecreate)
+	for (const FFastGeoRegisteredComponent& PendingRecreateComponent : ComponentsPendingRecreate)
 	{
 		// Skip component if its container was unregistered or re-registered before we process it
-		if (FFastGeoPrimitiveComponent* Component = PendingRecreateComponent.TryGetRegistered<FFastGeoPrimitiveComponent>())
+		if (FFastGeoComponent* Component = PendingRecreateComponent.TryGetRegistered<FFastGeoComponent>())
 		{
 #if CSV_PROFILER && !CSV_PROFILER_MINIMAL
 			if (Component->IsRenderStateDelayed())
diff --git a/Engine/Plugins/Experimental/FastGeoStreaming/Source/FastGeoStreaming/Private/FastGeoWorldSubsystem.h b/Engine/Plugins/Experimental/FastGeoStreaming/Source/FastGeoStreaming/Private/FastGeoWorldSubsystem.h
--- a/Engine/Plugins/Experimental/FastGeoStreaming/Source/FastGeoStreaming/Private/FastGeoWorldSubsystem.h
+++ b/Engine/Plugins/Experimental/FastGeoStreaming/Source/FastGeoStreaming/Private/FastGeoWorldSubsystem.h
@@ -4,6 +4,7 @@
 
 #include "Subsystems/WorldSubsystem.h"
 #include "FastGeoWeakElement.h"
+#include "FastGeoRegisteredComponent.h"
 #include "FastGeoAsyncRenderStateJobQueue.h"
 #include "FastGeoWorldSubsystem.generated.h"
 
@@ -36,7 +37,7 @@ public:
 	virtual TStatId GetStatId() const override;
 	//~ End UTickableWorldSubsystem interface
 
-	void AddToComponentsPendingRecreate(FFastGeoPrimitiveComponent* InComponentPendingRecreate);
+	void AddToComponentsPendingRecreate(FFastGeoComponent* InComponentPendingRecreate);
 	void ProcessPendingRecreate();
 
 	/**
@@ -71,6 +72,10 @@ public:
 
 	static bool IsEnableDebugView();
 
+#if !UE_BUILD_SHIPPING
+	static bool IsFastGeoVisible();
+#endif
+
 	bool IsWaitingForCompletion() const;
 
 private:
@@ -98,42 +103,9 @@ private:
 	FDelegateHandle Handle_OnLevelBeginRemoveFromWorld;
 	FDelegateHandle Handle_OnForEachHLODObjectInCell;
 
-	struct FRegisteredComponent
-	{
-		FRegisteredComponent(FFastGeoComponent* InComponent)
-			: Component(InComponent)
-		{
-			check(InComponent);
-			check(InComponent->GetOwnerContainer());
-			RegistrationEpoch = InComponent->GetOwnerContainer()->RegistrationEpoch;
-		}
-
-		template <class T>
-		FORCEINLINE T* TryGetRegistered() const
-		{
-			if (T* ComponentPtr = Component.Get<T>())
-			{
-				if (ComponentPtr->GetOwnerContainer()->IsRegistered() &&
-					RegistrationEpoch == ComponentPtr->GetOwnerContainer()->RegistrationEpoch)
-				{
-					return ComponentPtr;
-				}
-			}
-			return nullptr;
-		}
-
-		FORCEINLINE bool operator==(const FRegisteredComponent& Other) const
-		{
-			return Component == Other.Component && RegistrationEpoch == Other.RegistrationEpoch;
-		}
-
-	private:
-		FWeakFastGeoComponent Component;
-		int32 RegistrationEpoch = -1;
-	};
-
-	TArray<FRegisteredComponent> ComponentsPendingRecreate;
+	TArray<FFastGeoRegisteredComponent> ComponentsPendingRecreate;
 	static bool bEnableDebugView;
+	static bool bIsFastGeoVisible;
 
 	FRWLock Lock;
 	int32 TimeEpoch = 0;
diff --git a/Engine/Plugins/Experimental/FastGeoStreaming/Source/FastGeoStreaming/Public/FastGeoWorldPartitionRuntimeCellTransformer.h b/Engine/Plugins/Experimental/FastGeoStreaming/Source/FastGeoStreaming/Public/FastGeoWorldPartitionRuntimeCellTransformer.h
--- a/Engine/Plugins/Experimental/FastGeoStreaming/Source/FastGeoStreaming/Public/FastGeoWorldPartitionRuntimeCellTransformer.h
+++ b/Engine/Plugins/Experimental/FastGeoStreaming/Source/FastGeoStreaming/Public/FastGeoWorldPartitionRuntimeCellTransformer.h
@@ -126,7 +126,7 @@ protected:
 	/** Whether the actor can be processed by the transformer. */
 	virtual bool IsActorTransformable(AActor* InActor, FString& OutReason) const { return true; }
 	/** Whether the component can be processed by the transformer. */
-	virtual bool IsComponentTransformable(UPrimitiveComponent* InComponent, FString& OutReason) const { return true; }
+	virtual bool IsComponentTransformable(UActorComponent* InComponent, FString& OutReason) const { return true; }
 	/** Whether the fully transformed actor can be deleted by the transformer. */
 	virtual bool IsFullyTransformedActorDeletable(AActor* InActor, FString& OutReason) const { return true; }
 
