[FastGeo] Export FFastGeoHLOD::Type

diff --git a/Engine/Plugins/Experimental/FastGeoStreaming/Source/FastGeoStreaming/Internal/FastGeoHLOD.h b/Engine/Plugins/Experimental/FastGeoStreaming/Source/FastGeoStreaming/Internal/FastGeoHLOD.h
--- a/Engine/Plugins/Experimental/FastGeoStreaming/Source/FastGeoStreaming/Internal/FastGeoHLOD.h
+++ b/Engine/Plugins/Experimental/FastGeoStreaming/Source/FastGeoStreaming/Internal/FastGeoHLOD.h
@@ -11,7 +11,7 @@ public:
 	typedef FFastGeoComponentCluster Super;
 
 	/** Static type identifier for this element class */
-	static const FFastGeoElementType Type;
+	FASTGEOSTREAMING_API static const FFastGeoElementType Type;
 
 	FFastGeoHLOD(UFastGeoContainer* InOwner = nullptr, FName InName = NAME_None, FFastGeoElementType InType = Type);
 	~FFastGeoHLOD() {}
