Enable FastGeo and PCG FastGeo Interop plugins to enable runtime primitive spawning for PCG grass via FastGeo components.

Reduces setup/tear down costs and object allocations.

#rnx

diff --git a/Engine/Plugins/Experimental/FastGeoStreaming/Source/FastGeoStreaming/Private/FastGeoContainer.cpp b/Engine/Plugins/Experimental/FastGeoStreaming/Source/FastGeoStreaming/Private/FastGeoContainer.cpp
--- a/Engine/Plugins/Experimental/FastGeoStreaming/Source/FastGeoStreaming/Private/FastGeoContainer.cpp
+++ b/Engine/Plugins/Experimental/FastGeoStreaming/Source/FastGeoStreaming/Private/FastGeoContainer.cpp
@@ -4,6 +4,7 @@
 #include "FastGeoComponentCluster.h"
 #include "FastGeoHLOD.h"
 #include "FastGeoComponent.h"
+#include "FastGeoStreamingModule.h"
 #include "FastGeoWorldSubsystem.h"
 #include "FastGeoDestroyRenderStateContext.h"
 #include "FastGeoLog.h"
@@ -168,6 +169,11 @@ void UFastGeoContainer::Register()
 		return;
 	}
 #endif
+	if (!FFastGeoStreamingModule::IsFastGeoEnabled())
+	{
+		UE_LOG(LogFastGeoStreaming, Warning, TEXT("UFastGeoContainer::Register() ignored because FastGeo is disabled."));
+		return;
+	}
 
 	if (!bIsRegistered)
 	{
diff --git a/Engine/Plugins/Experimental/FastGeoStreaming/Source/FastGeoStreaming/Private/FastGeoStreamingModule.cpp b/Engine/Plugins/Experimental/FastGeoStreaming/Source/FastGeoStreaming/Private/FastGeoStreamingModule.cpp
--- a/Engine/Plugins/Experimental/FastGeoStreaming/Source/FastGeoStreaming/Private/FastGeoStreamingModule.cpp
+++ b/Engine/Plugins/Experimental/FastGeoStreaming/Source/FastGeoStreaming/Private/FastGeoStreamingModule.cpp
@@ -5,11 +5,19 @@
 #include "FastGeoStaticMeshComponent.h"
 #include "FastGeoInstancedStaticMeshComponent.h"
 #include "FastGeoWorldSubsystem.h"
+#include "HAL/IConsoleManager.h"
 #include "Modules/ModuleManager.h"
 #include "Engine/World.h"
 
 IMPLEMENT_MODULE(FFastGeoStreamingModule, FastGeoStreaming);
 
+bool FFastGeoStreamingModule::bIsFastGeoEnabled = true;
+FAutoConsoleVariableRef FFastGeoStreamingModule::CVarIsEnabled(
+	TEXT("FastGeo.Enable"),
+	FFastGeoStreamingModule::bIsFastGeoEnabled,
+	TEXT("Set to false to disable FastGeo (used in PIE and at cook time)."),
+	ECVF_Default | ECVF_ReadOnly);
+
 void FFastGeoStreamingModule::StartupModule()
 {
 	IPrimitiveComponent::AddImplementer(FComponentInterfaceImplementation(UFastGeoContainer::StaticClass(),  [](UObject* Obj)
@@ -33,6 +41,11 @@ void FFastGeoStreamingModule::ShutdownModule()
 	FWorldDelegates::OnWorldPreSendAllEndOfFrameUpdates.Remove(Handle_OnWorldPreSendAllEndOfFrameUpdates);
 }
 
+bool FFastGeoStreamingModule::IsFastGeoEnabled()
+{
+	return bIsFastGeoEnabled;
+}
+
 void FFastGeoStreamingModule::OnWorldPreSendAllEndOfFrameUpdates(UWorld* InWorld)
 {
 	if (UFastGeoWorldSubsystem* WorldSubsystem = InWorld->GetSubsystem<UFastGeoWorldSubsystem>())
diff --git a/Engine/Plugins/Experimental/FastGeoStreaming/Source/FastGeoStreaming/Private/FastGeoWorldPartitionRuntimeCellTransformer.cpp b/Engine/Plugins/Experimental/FastGeoStreaming/Source/FastGeoStreaming/Private/FastGeoWorldPartitionRuntimeCellTransformer.cpp
--- a/Engine/Plugins/Experimental/FastGeoStreaming/Source/FastGeoStreaming/Private/FastGeoWorldPartitionRuntimeCellTransformer.cpp
+++ b/Engine/Plugins/Experimental/FastGeoStreaming/Source/FastGeoStreaming/Private/FastGeoWorldPartitionRuntimeCellTransformer.cpp
@@ -2,6 +2,7 @@
 
 #include "FastGeoWorldPartitionRuntimeCellTransformer.h"
 #include "UObject/FortniteMainBranchObjectVersion.h"
+#include "FastGeoStreamingModule.h"
 
 #if WITH_EDITOR
 #include "CoreMinimal.h"
@@ -182,11 +183,11 @@ FAutoConsoleVariableRef UFastGeoWorldPartitionRuntimeCellTransformer::CVarIsDebu
 	TEXT("Set to true to enable FastGeoStreaming transformer debug mode (used in PIE and at cook time)."),
 	ECVF_Default);
 
-bool UFastGeoWorldPartitionRuntimeCellTransformer::IsFastGeoEnabled = true;
-FAutoConsoleVariableRef UFastGeoWorldPartitionRuntimeCellTransformer::CVarIsFastGeoEnabled(
-	TEXT("FastGeo.Enable"),
-	UFastGeoWorldPartitionRuntimeCellTransformer::IsFastGeoEnabled,
-	TEXT("Set to false to disable FastGeoStreaming (used in PIE and at cook time)."),
+bool UFastGeoWorldPartitionRuntimeCellTransformer::IsFastGeoTransformerEnabled = true;
+FAutoConsoleVariableRef UFastGeoWorldPartitionRuntimeCellTransformer::CVarIsFastGeoTransformerEnabled(
+	TEXT("FastGeo.EnableTransformer"),
+	UFastGeoWorldPartitionRuntimeCellTransformer::IsFastGeoTransformerEnabled,
+	TEXT("Set to false to disable FastGeo cell transformer (used in PIE and at cook time). Has no effect if FastGeo.Enable is false."),
 	ECVF_Default);
 
 void UFastGeoWorldPartitionRuntimeCellTransformer::BeginDestroy()
@@ -282,7 +283,7 @@ void UFastGeoWorldPartitionRuntimeCellTransformer::Transform(ULevel* InLevel)
 {
 	TGuardValue<bool> Guard(FFastGeoTransformResult::ShouldReport, IsDebugMode());
 
-	if (!UFastGeoWorldPartitionRuntimeCellTransformer::IsFastGeoEnabled)
+	if (!FFastGeoStreamingModule::IsFastGeoEnabled() || !UFastGeoWorldPartitionRuntimeCellTransformer::IsFastGeoTransformerEnabled)
 	{
 		return;
 	}
diff --git a/Engine/Plugins/Experimental/FastGeoStreaming/Source/FastGeoStreaming/Private/FastGeoWorldSubsystem.cpp b/Engine/Plugins/Experimental/FastGeoStreaming/Source/FastGeoStreaming/Private/FastGeoWorldSubsystem.cpp
--- a/Engine/Plugins/Experimental/FastGeoStreaming/Source/FastGeoStreaming/Private/FastGeoWorldSubsystem.cpp
+++ b/Engine/Plugins/Experimental/FastGeoStreaming/Source/FastGeoStreaming/Private/FastGeoWorldSubsystem.cpp
@@ -3,6 +3,7 @@
 #include "FastGeoWorldSubsystem.h"
 #include "FastGeoContainer.h"
 #include "FastGeoPrimitiveComponent.h"
+#include "FastGeoStreamingModule.h"
 #include "Physics/Experimental/PhysScene_Chaos.h"
 #include "Streaming/LevelStreamingDelegates.h"
 #include "Engine/Engine.h"
@@ -82,6 +83,11 @@ bool UFastGeoWorldSubsystem::IsFastGeoVisible()
 
 UFastGeoWorldSubsystem::UFastGeoWorldSubsystem()
 {
+	if (!FFastGeoStreamingModule::IsFastGeoEnabled())
+	{
+		return;
+	}
+
 #if ENABLE_ACTOR_PRIMITIVE_COLOR_HANDLER
 
 	auto UpdatePrimitivesColor = []()
@@ -142,6 +148,11 @@ UFastGeoWorldSubsystem::UFastGeoWorldSubsystem()
 
 UFastGeoWorldSubsystem::~UFastGeoWorldSubsystem()
 {
+	if (!FFastGeoStreamingModule::IsFastGeoEnabled())
+	{
+		return;
+	}
+
 #if ENABLE_ACTOR_PRIMITIVE_COLOR_HANDLER
 	if (HasAnyFlags(RF_ClassDefaultObject) && ExactCast<UFastGeoWorldSubsystem>(this))
 	{
@@ -152,7 +163,7 @@ UFastGeoWorldSubsystem::~UFastGeoWorldSubsystem()
 
 bool UFastGeoWorldSubsystem::DoesSupportWorldType(const EWorldType::Type WorldType) const
 {
-	return WorldType == EWorldType::Game || WorldType == EWorldType::PIE;
+	return FFastGeoStreamingModule::IsFastGeoEnabled() && (WorldType == EWorldType::Game || WorldType == EWorldType::PIE);
 }
 
 void UFastGeoWorldSubsystem::PostInitialize()
diff --git a/Engine/Plugins/Experimental/FastGeoStreaming/Source/FastGeoStreaming/Public/FastGeoStreamingModule.h b/Engine/Plugins/Experimental/FastGeoStreaming/Source/FastGeoStreaming/Public/FastGeoStreamingModule.h
--- a/Engine/Plugins/Experimental/FastGeoStreaming/Source/FastGeoStreaming/Public/FastGeoStreamingModule.h
+++ b/Engine/Plugins/Experimental/FastGeoStreaming/Source/FastGeoStreaming/Public/FastGeoStreamingModule.h
@@ -5,6 +5,7 @@
 #include "Delegates/Delegate.h"
 #include "Modules/ModuleInterface.h"
 
+class FAutoConsoleVariableRef;
 class UWorld;
 
 class FFastGeoStreamingModule : public IModuleInterface
@@ -16,9 +17,14 @@ public:
 	virtual void StartupModule() override;
 	virtual void ShutdownModule() override;
 
+	FASTGEOSTREAMING_API static bool IsFastGeoEnabled();
+
 private:
 	void OnWorldPreSendAllEndOfFrameUpdates(UWorld* InWorld);
 
 private:
 	FDelegateHandle Handle_OnWorldPreSendAllEndOfFrameUpdates;
+
+	static bool bIsFastGeoEnabled;
+	static FAutoConsoleVariableRef CVarIsEnabled;
 };
diff --git a/Engine/Plugins/Experimental/FastGeoStreaming/Source/FastGeoStreaming/Public/FastGeoWorldPartitionRuntimeCellTransformer.h b/Engine/Plugins/Experimental/FastGeoStreaming/Source/FastGeoStreaming/Public/FastGeoWorldPartitionRuntimeCellTransformer.h
--- a/Engine/Plugins/Experimental/FastGeoStreaming/Source/FastGeoStreaming/Public/FastGeoWorldPartitionRuntimeCellTransformer.h
+++ b/Engine/Plugins/Experimental/FastGeoStreaming/Source/FastGeoStreaming/Public/FastGeoWorldPartitionRuntimeCellTransformer.h
@@ -218,9 +218,9 @@ protected:
 
 #if WITH_EDITOR
 	static bool IsDebugModeEnabled;
-	static bool IsFastGeoEnabled;
+	static bool IsFastGeoTransformerEnabled;
 	static FAutoConsoleVariableRef CVarIsDebugModeEnabled;
-	static FAutoConsoleVariableRef CVarIsFastGeoEnabled;
+	static FAutoConsoleVariableRef CVarIsFastGeoTransformerEnabled;
 #endif
 
 #if WITH_EDITORONLY_DATA
diff --git a/Engine/Plugins/Experimental/PCGInterops/PCGFastGeoInterop/Source/PCGFastGeoInterop/Private/Compute/PrimitiveFactories/PCGPrimitiveFactoryFastGeoPISMC.cpp b/Engine/Plugins/Experimental/PCGInterops/PCGFastGeoInterop/Source/PCGFastGeoInterop/Private/Compute/PrimitiveFactories/PCGPrimitiveFactoryFastGeoPISMC.cpp
--- a/Engine/Plugins/Experimental/PCGInterops/PCGFastGeoInterop/Source/PCGFastGeoInterop/Private/Compute/PrimitiveFactories/PCGPrimitiveFactoryFastGeoPISMC.cpp
+++ b/Engine/Plugins/Experimental/PCGInterops/PCGFastGeoInterop/Source/PCGFastGeoInterop/Private/Compute/PrimitiveFactories/PCGPrimitiveFactoryFastGeoPISMC.cpp
@@ -7,6 +7,8 @@
 #include "PCGComponent.h"
 #include "PCGContext.h"
 #include "PCGModule.h"
+
+#include "FastGeoStreamingModule.h"
 #include "PrimitiveSceneInfo.h"
 #include "Components/PCGProceduralISMComponentDescriptor.h"
 #include "Compute/PCGComputeCommon.h"
@@ -74,6 +76,12 @@ void FPCGPrimitiveFactoryFastGeoPISMC::Initialize(FParameters&& InParameters)
 
 bool FPCGPrimitiveFactoryFastGeoPISMC::Create(FPCGContext* InContext)
 {
+	if (!FFastGeoStreamingModule::IsFastGeoEnabled())
+	{
+		UE_LOG(LogPCG, Warning, TEXT("FPCGPrimitiveFactoryFastGeoPISMC: FastGeoStreaming plugin is not enabled, no primitives will be spawned."));
+		return true;
+	}
+
 	IPCGGraphExecutionSource* ExecutionSource = InContext->ExecutionSource.Get();
 	if (!ExecutionSource)
 	{
