[PCG][FastGeo] pcg.RuntimeGeneration.ISM.ComponentlessPrimitive editor support

diff --git a/Engine/Plugins/Experimental/FastGeoStreaming/Source/FastGeoStreaming/Private/FastGeoAsyncRenderStateJobQueue.cpp b/Engine/Plugins/Experimental/FastGeoStreaming/Source/FastGeoStreaming/Private/FastGeoAsyncRenderStateJobQueue.cpp
--- a/Engine/Plugins/Experimental/FastGeoStreaming/Source/FastGeoStreaming/Private/FastGeoAsyncRenderStateJobQueue.cpp
+++ b/Engine/Plugins/Experimental/FastGeoStreaming/Source/FastGeoStreaming/Private/FastGeoAsyncRenderStateJobQueue.cpp
@@ -122,6 +122,14 @@ void FFastGeoAsyncRenderStateJobQueue::Tick(bool bForceWaitCompletion)
 	// Flush all queued jobs
 	Launch();
 
+#if WITH_EDITOR
+	// In non-game worlds make sure event is triggered as OnUpdateLevelStreamingDone isn't called
+	if (!Scene->GetWorld()->IsGameWorld())
+	{
+		TriggerIsReadyToRunAsyncTasksEvent();
+	}
+#endif
+
 	// Wait for tasks if the world is inside a blocking load
 	if (bForceWaitCompletion || Scene->GetWorld()->GetSubsystem<UFastGeoWorldSubsystem>()->IsWaitingForCompletion())
 	{
diff --git a/Engine/Plugins/Experimental/FastGeoStreaming/Source/FastGeoStreaming/Private/FastGeoContainer.cpp b/Engine/Plugins/Experimental/FastGeoStreaming/Source/FastGeoStreaming/Private/FastGeoContainer.cpp
--- a/Engine/Plugins/Experimental/FastGeoStreaming/Source/FastGeoStreaming/Private/FastGeoContainer.cpp
+++ b/Engine/Plugins/Experimental/FastGeoStreaming/Source/FastGeoStreaming/Private/FastGeoContainer.cpp
@@ -180,7 +180,6 @@ void UFastGeoContainer::Register()
 		++RegistrationEpoch;
 		PendingCreate.Reset();
 
-		check(GetWorld()->IsGameWorld());
 #if WITH_EDITOR
 		// In PIE we need to initialize dynamic properties as there's no serialization
 		InitializeDynamicProperties(true);
diff --git a/Engine/Plugins/Experimental/FastGeoStreaming/Source/FastGeoStreaming/Private/FastGeoWorldSubsystem.cpp b/Engine/Plugins/Experimental/FastGeoStreaming/Source/FastGeoStreaming/Private/FastGeoWorldSubsystem.cpp
--- a/Engine/Plugins/Experimental/FastGeoStreaming/Source/FastGeoStreaming/Private/FastGeoWorldSubsystem.cpp
+++ b/Engine/Plugins/Experimental/FastGeoStreaming/Source/FastGeoStreaming/Private/FastGeoWorldSubsystem.cpp
@@ -163,7 +163,7 @@ UFastGeoWorldSubsystem::~UFastGeoWorldSubsystem()
 
 bool UFastGeoWorldSubsystem::DoesSupportWorldType(const EWorldType::Type WorldType) const
 {
-	return FFastGeoStreamingModule::IsFastGeoEnabled() && (WorldType == EWorldType::Game || WorldType == EWorldType::PIE);
+	return FFastGeoStreamingModule::IsFastGeoEnabled() && Super::DoesSupportWorldType(WorldType);
 }
 
 void UFastGeoWorldSubsystem::PostInitialize()
@@ -171,7 +171,7 @@ void UFastGeoWorldSubsystem::PostInitialize()
 	Super::PostInitialize();
 
 	UWorld* World = GetWorld();
-	if (World->IsPartitionedWorld())
+	if (World->IsPartitionedWorld() && World->IsGameWorld())
 	{
 		World->OnAllLevelsChanged().AddUObject(this, &UFastGeoWorldSubsystem::OnUpdateLevelStreaming);
 		World->OnAddLevelToWorldExtension().AddUObject(this, &UFastGeoWorldSubsystem::OnAddLevelToWorldExtension);
@@ -193,24 +193,31 @@ void UFastGeoWorldSubsystem::PostInitialize()
 void UFastGeoWorldSubsystem::Deinitialize()
 {
 	UWorld* World = GetWorld();
-	World->OnAllLevelsChanged().RemoveAll(this);
-	World->OnAddLevelToWorldExtension().RemoveAll(this);
-	World->OnRemoveLevelFromWorldExtension().RemoveAll(this);
-	FWorldDelegates::LevelComponentsUpdated.RemoveAll(this);
-	FWorldDelegates::LevelComponentsCleared.RemoveAll(this);
+
+	if (World->IsGameWorld())
+	{
+		World->OnAllLevelsChanged().RemoveAll(this);
+		World->OnAddLevelToWorldExtension().RemoveAll(this);
+		World->OnRemoveLevelFromWorldExtension().RemoveAll(this);
+		FWorldDelegates::LevelComponentsUpdated.RemoveAll(this);
+		FWorldDelegates::LevelComponentsCleared.RemoveAll(this);
 #if DO_CHECK
-	FWorldDelegates::LevelAddedToWorld.RemoveAll(this);
-	FWorldDelegates::LevelRemovedFromWorld.RemoveAll(this);
+		FWorldDelegates::LevelAddedToWorld.RemoveAll(this);
+		FWorldDelegates::LevelRemovedFromWorld.RemoveAll(this);
 #endif
-	FLevelStreamingDelegates::OnLevelStreamingStateChanged.Remove(Handle_OnLevelStreamingStateChanged);
-	FLevelStreamingDelegates::OnLevelBeginMakingVisible.Remove(Handle_OnLevelBeginAddToWorld);
-	FLevelStreamingDelegates::OnLevelBeginMakingInvisible.Remove(Handle_OnLevelBeginRemoveFromWorld);
-	Handle_OnLevelStreamingStateChanged.Reset();
-	Handle_OnLevelBeginAddToWorld.Reset();
-	Handle_OnLevelBeginRemoveFromWorld.Reset();
-
-	World->GetSubsystem<UWorldPartitionHLODRuntimeSubsystem>()->GetForEachHLODObjectInCellEvent().Remove(Handle_OnForEachHLODObjectInCell);
-	Handle_OnForEachHLODObjectInCell.Reset();
+		FLevelStreamingDelegates::OnLevelStreamingStateChanged.Remove(Handle_OnLevelStreamingStateChanged);
+		FLevelStreamingDelegates::OnLevelBeginMakingVisible.Remove(Handle_OnLevelBeginAddToWorld);
+		FLevelStreamingDelegates::OnLevelBeginMakingInvisible.Remove(Handle_OnLevelBeginRemoveFromWorld);
+		Handle_OnLevelStreamingStateChanged.Reset();
+		Handle_OnLevelBeginAddToWorld.Reset();
+		Handle_OnLevelBeginRemoveFromWorld.Reset();
+
+		if (Handle_OnForEachHLODObjectInCell.IsValid())
+		{
+			World->GetSubsystem<UWorldPartitionHLODRuntimeSubsystem>()->GetForEachHLODObjectInCellEvent().Remove(Handle_OnForEachHLODObjectInCell);
+			Handle_OnForEachHLODObjectInCell.Reset();
+		}
+	}
 
 	Super::Deinitialize();
 }
diff --git a/Engine/Plugins/Experimental/PCGInterops/PCGFastGeoInterop/Source/PCGFastGeoInterop/Private/Compute/PrimitiveFactories/PCGPrimitiveFactoryFastGeoPISMC.cpp b/Engine/Plugins/Experimental/PCGInterops/PCGFastGeoInterop/Source/PCGFastGeoInterop/Private/Compute/PrimitiveFactories/PCGPrimitiveFactoryFastGeoPISMC.cpp
--- a/Engine/Plugins/Experimental/PCGInterops/PCGFastGeoInterop/Source/PCGFastGeoInterop/Private/Compute/PrimitiveFactories/PCGPrimitiveFactoryFastGeoPISMC.cpp
+++ b/Engine/Plugins/Experimental/PCGInterops/PCGFastGeoInterop/Source/PCGFastGeoInterop/Private/Compute/PrimitiveFactories/PCGPrimitiveFactoryFastGeoPISMC.cpp
@@ -89,7 +89,7 @@ bool FPCGPrimitiveFactoryFastGeoPISMC::Create(FPCGContext* InContext)
 	}
 
 	UWorld* World = InContext->ExecutionSource->GetExecutionState().GetWorld();
-	ULevel* Level = World ? World->GetCurrentLevel() : nullptr;
+	ULevel* Level = World ? World->PersistentLevel : nullptr;
 
 	if (!Level)
 	{
@@ -97,13 +97,7 @@ bool FPCGPrimitiveFactoryFastGeoPISMC::Create(FPCGContext* InContext)
 		return true;
 	}
 
-	// Currently only support transient editor world.
-	if (!ensure(InContext->ExecutionSource->GetExecutionState().GetWorld()->IsGameWorld()))
-	{
-		return true;
-	}
-
-	// Must be outer'd to level - so pick same level as pcg component.
+	// Must be outer'd to level - so pick world persistent level.
 	UFastGeoContainer* FastGeo = FPCGContext::NewObject_AnyThread<UFastGeoContainer>(InContext, Level);
 
 	const FString DebugName = ExecutionSource->GetExecutionState().GetDebugName();
diff --git a/Engine/Plugins/PCG/Source/PCG/Private/Compute/DataInterfaces/Elements/PCGStaticMeshSpawnerDataInterface.cpp b/Engine/Plugins/PCG/Source/PCG/Private/Compute/DataInterfaces/Elements/PCGStaticMeshSpawnerDataInterface.cpp
--- a/Engine/Plugins/PCG/Source/PCG/Private/Compute/DataInterfaces/Elements/PCGStaticMeshSpawnerDataInterface.cpp
+++ b/Engine/Plugins/PCG/Source/PCG/Private/Compute/DataInterfaces/Elements/PCGStaticMeshSpawnerDataInterface.cpp
@@ -956,7 +956,7 @@ bool UPCGStaticMeshSpawnerDataProvider::SetupPrimitives(FPCGContext* InContext,
 		{
 			if (UWorld* World = InContext->ExecutionSource.IsValid() ? InContext->ExecutionSource->GetExecutionState().GetWorld() : nullptr)
 			{
-				bCreateComponentless = World->IsGameWorld();
+				bCreateComponentless = true;
 			}
 		}
 
