Cleanup IPrimitiveComponent interface implementation, removed the intermediate struct instance interfacing with FastGeoComponent.

#rnx

diff --git a/Engine/Plugins/Experimental/FastGeoStreaming/Source/FastGeoStreaming/Internal/FastGeoComponent.h b/Engine/Plugins/Experimental/FastGeoStreaming/Source/FastGeoStreaming/Internal/FastGeoComponent.h
--- a/Engine/Plugins/Experimental/FastGeoStreaming/Source/FastGeoStreaming/Internal/FastGeoComponent.h
+++ b/Engine/Plugins/Experimental/FastGeoStreaming/Source/FastGeoStreaming/Internal/FastGeoComponent.h
@@ -28,12 +28,6 @@ enum class ESceneProxyCreationError
 	InvalidMesh
 };
 
-#define FASTGEO_DECLARE_COMPONENT_INTERFACE_INTERNAL(ComponentName, BasePrefix, ...)\
-	UE_DECLARE_COMPONENT_INTERFACE_INTERNAL(BasePrefix##ComponentName##__VA_ARGS__, I##ComponentName , BasePrefix##ComponentName##Interface, ComponentName##Interface, ComponentName)
-
-#define FASTGEO_DECLARE_COMPONENT_INTERFACE(ComponentName, ...)\
-	FASTGEO_DECLARE_COMPONENT_INTERFACE_INTERNAL(ComponentName, FFastGeo, __VA_ARGS__)
-
 class FASTGEOSTREAMING_API FFastGeoComponent : public IFastGeoElement
 {
 public:
@@ -114,4 +108,4 @@ private:
 	friend class FFastGeoComponentCluster;
 
 	friend FArchive& operator<<(FArchive& Ar, FFastGeoComponent& Component);
-};
+};
\ No newline at end of file
diff --git a/Engine/Plugins/Experimental/FastGeoStreaming/Source/FastGeoStreaming/Internal/FastGeoContainer.h b/Engine/Plugins/Experimental/FastGeoStreaming/Source/FastGeoStreaming/Internal/FastGeoContainer.h
--- a/Engine/Plugins/Experimental/FastGeoStreaming/Source/FastGeoStreaming/Internal/FastGeoContainer.h
+++ b/Engine/Plugins/Experimental/FastGeoStreaming/Source/FastGeoStreaming/Internal/FastGeoContainer.h
@@ -116,8 +116,8 @@ public:
 	ULevel* GetLevel() const;
 	FFastGeoComponentCluster* GetComponentCluster(uint32 InComponentClusterTypeID, int32 InComponentClusterIndex);
 
-	TArray<void*> GetPrimitiveComponents() const;
-	TArray<void*> GetStaticMeshComponents() const;
+	TArray<IPrimitiveComponent*> GetPrimitiveComponents() const;
+	TArray<IStaticMeshComponent*> GetStaticMeshComponents() const;
 
 	//~ Begin UObject Interface
 	virtual void BeginDestroy() override;
@@ -286,8 +286,8 @@ private:
 #endif
 
 	// Transient data
-	TArray<void*> PrimitiveComponents;
-	TArray<void*> StaticMeshComponents;
+	TArray<IPrimitiveComponent*> PrimitiveComponents;
+	TArray<IStaticMeshComponent*> StaticMeshComponents;
 	TArray<FFastGeoComponent*> CollisionComponents;
 	TArray<FNavigationElementHandle> NavigationElementHandles;
 
diff --git a/Engine/Plugins/Experimental/FastGeoStreaming/Source/FastGeoStreaming/Internal/FastGeoPrimitiveComponent.h b/Engine/Plugins/Experimental/FastGeoStreaming/Source/FastGeoStreaming/Internal/FastGeoPrimitiveComponent.h
--- a/Engine/Plugins/Experimental/FastGeoStreaming/Source/FastGeoStreaming/Internal/FastGeoPrimitiveComponent.h
+++ b/Engine/Plugins/Experimental/FastGeoStreaming/Source/FastGeoStreaming/Internal/FastGeoPrimitiveComponent.h
@@ -3,7 +3,6 @@
 #pragma once
 
 #include "FastGeoComponent.h"
-#include "FastGeoPrimitiveComponentInterface.h"
 #include "Components/PrimitiveComponent.h"
 #include "Engine/EngineTypes.h"
 #include "PrimitiveSceneInfoData.h"
@@ -59,9 +58,8 @@ private:
 	friend class FFastGeoInstancedStaticMeshComponent;
 };
 
-class FASTGEOSTREAMING_API FFastGeoPrimitiveComponent : public FFastGeoComponent
+class FASTGEOSTREAMING_API FFastGeoPrimitiveComponent : public FFastGeoComponent, public IPrimitiveComponent
 {
-	FASTGEO_DECLARE_COMPONENT_INTERFACE(PrimitiveComponent)
 public:
 	typedef FFastGeoComponent Super;
 
@@ -78,14 +76,10 @@ public:
 	FPrimitiveComponentId GetPrimitiveSceneId() const { return PrimitiveSceneData.PrimitiveSceneId; }
 	virtual UObject const* AdditionalStatObject() const { return nullptr; }
 		
-	FPrimitiveSceneProxy* GetSceneProxy() const;
-	const FTransform& GetTransform() const;
-	const FBoxSphereBounds& GetBounds() const;
 	FMatrix GetRenderMatrix() const;
-	float GetLastRenderTimeOnScreen() const;
 
 	// Used by FStaticMeshComponentHelper/FInstancedStaticMeshComponentHelper
-	const FTransform& GetComponentTransform() const { return GetTransform(); }
+	FTransform GetComponentTransform() const { return GetTransform(); }
 
 #if WITH_EDITOR
 	virtual void NotifyRenderStateChanged();
@@ -96,7 +90,38 @@ public:
 	void SetCollisionEnabled(bool bInCollisionEnabled);
 	void UpdateVisibility();
 	EComponentMobility::Type GetMobility() const;
-	
+
+	//~ Begin IPrimitiveComponent interface
+	virtual bool IsRenderStateDirty() const override;
+	virtual bool IsRenderStateCreated() const override;
+	virtual bool ShouldCreateRenderState() const override;
+	virtual bool IsRegistered() const override;
+	virtual bool IsUnreachable() const override;
+	virtual UWorld* GetWorld() const override;
+	virtual FSceneInterface* GetScene() const override;
+	virtual FPrimitiveSceneProxy* GetSceneProxy() const override;
+	virtual void GetStreamableRenderAssetInfo(TArray<FStreamingRenderAssetPrimitiveInfo>&) const override;
+	virtual void MarkRenderStateDirty() override;
+	virtual void DestroyRenderState() override;
+	virtual void CreateRenderState(FRegisterComponentContext* Context) override;
+	virtual FString GetName() const override;
+	virtual FString GetFullName() const override;
+	virtual FTransform GetTransform() const override;
+	virtual FBoxSphereBounds GetBounds() const override;
+	float GetLastRenderTimeOnScreen() const override;
+	virtual void GetPrimitiveStats(FPrimitiveStats& PrimitiveStats) const override;
+	virtual UObject* GetUObject() override;
+	virtual const UObject* GetUObject() const override;
+	virtual void PrecachePSOs() override;
+	virtual UObject* GetOwner() const override;
+	virtual FString GetOwnerName() const override;
+	virtual FPrimitiveSceneProxy* CreateSceneProxy() override;
+#if WITH_EDITOR
+	virtual HHitProxy* CreateMeshHitProxy(int32 SectionIndex, int32 MaterialIndex) override;
+#endif // WITH_EDITOR
+	virtual HHitProxy* CreatePrimitiveHitProxies(TArray<TRefCountPtr<HHitProxy> >& OutHitProxies) override;
+	//~ End IPrimitiveComponent interface
+
 	//~ Begin FFastGeoComponent interface
 	virtual void Serialize(FArchive& Ar) override;
 	virtual void InitializeDynamicProperties() override;
@@ -106,11 +131,8 @@ public:
 	virtual void OnAsyncDestroyPhysicsState() override;
 	virtual bool IsCollisionEnabled() const override;
 	virtual void ApplyWorldTransform(const FTransform& InTransform) override;
-	virtual bool ShouldCreateRenderState() const override;
-	virtual void CreateRenderState(FRegisterComponentContext* Context) override;
 	virtual void DestroyRenderState(FFastGeoDestroyRenderStateContext* Context) override;
 	virtual void MarkPrecachePSOsRequired() override;
-	virtual void PrecachePSOs() override;
 #if WITH_EDITOR
 	virtual void InitializeFromComponent(UActorComponent* Component) override;
 #endif
@@ -119,7 +141,6 @@ public:
 	//~ Begin Materials
 	virtual UMaterialInterface* GetMaterial(int32 ElementIndex) const = 0;
 	virtual int32 GetNumMaterials() const = 0;
-	virtual void GetUsedMaterials(TArray<UMaterialInterface*>& OutMaterials, bool bGetDebugMaterials = false) const = 0;
 	UE_DEPRECATED(5.7, "Please use GetUsedMaterialPropertyDesc with EShaderPlatform argument and not ERHIFeatureLevel::Type")
 	FPrimitiveMaterialPropertyDescriptor GetUsedMaterialPropertyDesc(ERHIFeatureLevel::Type FeatureLevel) const;
 	FPrimitiveMaterialPropertyDescriptor GetUsedMaterialPropertyDesc(EShaderPlatform InShaderPlatform) const;
@@ -143,7 +164,7 @@ protected:
 	virtual void ResetSceneProxyDescUnsupportedProperties();
 #endif
 	virtual void InitializeSceneProxyDescDynamicProperties();
-	virtual FPrimitiveSceneProxy* CreateSceneProxy(ESceneProxyCreationError* OutError = nullptr) = 0;
+	virtual FPrimitiveSceneProxy* CreateSceneProxy(ESceneProxyCreationError* OutError) = 0;
 	FPrimitiveSceneDesc BuildSceneDesc();
 
 	//~ Begin Physics
@@ -209,12 +230,3 @@ protected:
 	friend class FFastGeoPhysicsBodyInstanceOwner;
 	friend class UFastGeoContainer;
 };
-
-// Dummy type to use as base class in non editor builds
-template <typename TInterface>
-class IFastGeoDummyInterface {};
-
-// Implement TInterface, only in editor.
-// Would have been easier with a #ifdef in the class declaration, but UHT doesn't allow it.
-template <typename TInterface>
-class IFastGeoEditorInterface : public std::conditional<WITH_EDITOR, TInterface, IFastGeoDummyInterface<TInterface>>::type {};
diff --git a/Engine/Plugins/Experimental/FastGeoStreaming/Source/FastGeoStreaming/Internal/FastGeoPrimitiveComponentInterface.h b/Engine/Plugins/Experimental/FastGeoStreaming/Source/FastGeoStreaming/Internal/FastGeoPrimitiveComponentInterface.h
deleted file mode 100644
diff --git a/Engine/Plugins/Experimental/FastGeoStreaming/Source/FastGeoStreaming/Internal/FastGeoStaticMeshComponent.h b/Engine/Plugins/Experimental/FastGeoStreaming/Source/FastGeoStreaming/Internal/FastGeoStaticMeshComponent.h
--- a/Engine/Plugins/Experimental/FastGeoStreaming/Source/FastGeoStreaming/Internal/FastGeoStaticMeshComponent.h
+++ b/Engine/Plugins/Experimental/FastGeoStreaming/Source/FastGeoStreaming/Internal/FastGeoStaticMeshComponent.h
@@ -4,7 +4,6 @@
 
 #include "FastGeoMeshComponent.h"
 #include "StaticMeshSceneProxyDesc.h"
-#include "FastGeoStaticMeshComponentInterface.h"
 
 class UBodySetup;
 class UMaterialInterface;
@@ -20,10 +19,8 @@ namespace Nanite
 	class FNaniteResourcesHelper;
 }
 
-class FASTGEOSTREAMING_API FFastGeoStaticMeshComponentBase : public FFastGeoMeshComponent
+class FASTGEOSTREAMING_API FFastGeoStaticMeshComponentBase : public FFastGeoMeshComponent, public IStaticMeshComponent
 {
-	FASTGEO_DECLARE_COMPONENT_INTERFACE(StaticMeshComponent, Base)
-
 public:
 	typedef FFastGeoMeshComponent Super;
 
@@ -36,7 +33,6 @@ public:
 	const FCollisionResponseContainer& GetCollisionResponseToChannels() const;
 	UMaterialInterface* GetNaniteAuditMaterial(int32 MaterialIndex) const;
 	const Nanite::FResources* GetNaniteResources() const;
-	UStaticMesh* GetStaticMesh() const;
 	bool HasValidNaniteData() const;
 	bool UseNaniteOverrideMaterials(bool bDoingMaterialAudit) const;
 	bool ShouldCreateNaniteProxy(Nanite::FMaterialAudit* OutNaniteMaterials) const;
@@ -62,6 +58,16 @@ public:
 #endif
 	//~ End FFastGeoComponent interface
 
+	//~ Begin IStaticMeshComponent interface
+#if WITH_EDITOR
+	virtual void OnMeshRebuild(bool bRenderDataChanged) override {}
+	virtual void PreStaticMeshCompilation() override {}
+	virtual void PostStaticMeshCompilation() override {}
+#endif
+	virtual UStaticMesh* GetStaticMesh() const override;
+	virtual IPrimitiveComponent* GetPrimitiveComponentInterface() override;
+	//~ End IStaticMeshComponent interface
+
 	//~ Begin FFastGeoPrimitiveComponent interface
 	virtual bool IsNavigationRelevant() const override;
 	virtual FBox GetNavigationBounds() const override;
diff --git a/Engine/Plugins/Experimental/FastGeoStreaming/Source/FastGeoStreaming/Internal/FastGeoStaticMeshComponentInterface.h b/Engine/Plugins/Experimental/FastGeoStreaming/Source/FastGeoStreaming/Internal/FastGeoStaticMeshComponentInterface.h
deleted file mode 100644
diff --git a/Engine/Plugins/Experimental/FastGeoStreaming/Source/FastGeoStreaming/Private/FastGeoContainer.cpp b/Engine/Plugins/Experimental/FastGeoStreaming/Source/FastGeoStreaming/Private/FastGeoContainer.cpp
--- a/Engine/Plugins/Experimental/FastGeoStreaming/Source/FastGeoStreaming/Private/FastGeoContainer.cpp
+++ b/Engine/Plugins/Experimental/FastGeoStreaming/Source/FastGeoStreaming/Private/FastGeoContainer.cpp
@@ -328,12 +328,12 @@ FFastGeoComponentCluster* UFastGeoContainer::GetComponentCluster(uint32 InCompon
 	return nullptr;
 }
 
-TArray<void*> UFastGeoContainer::GetPrimitiveComponents() const
+TArray<IPrimitiveComponent*> UFastGeoContainer::GetPrimitiveComponents() const
 {
 	return PrimitiveComponents;
 }
 
-TArray<void*> UFastGeoContainer::GetStaticMeshComponents() const
+TArray<IStaticMeshComponent*> UFastGeoContainer::GetStaticMeshComponents() const
 {
 	return StaticMeshComponents;
 }
@@ -775,11 +775,11 @@ void UFastGeoContainer::InitializeDynamicProperties(bool bInitForPlay)
 
 				if (FFastGeoPrimitiveComponent* PrimitiveComponent = Component.CastTo<FFastGeoPrimitiveComponent>())
 				{
-					PrimitiveComponents.Add(PrimitiveComponent->GetPrimitiveComponentInterface());
+					PrimitiveComponents.Add(PrimitiveComponent);
 				}
 				if (FFastGeoStaticMeshComponent* StaticMeshComponent = Component.CastTo<FFastGeoStaticMeshComponent>())
 				{
-					StaticMeshComponents.Add(StaticMeshComponent->GetStaticMeshComponentInterface());
+					StaticMeshComponents.Add(StaticMeshComponent);
 				}
 			});
 		});
diff --git a/Engine/Plugins/Experimental/FastGeoStreaming/Source/FastGeoStreaming/Private/FastGeoPrimitiveComponent.cpp b/Engine/Plugins/Experimental/FastGeoStreaming/Source/FastGeoStreaming/Private/FastGeoPrimitiveComponent.cpp
--- a/Engine/Plugins/Experimental/FastGeoStreaming/Source/FastGeoStreaming/Private/FastGeoPrimitiveComponent.cpp
+++ b/Engine/Plugins/Experimental/FastGeoStreaming/Source/FastGeoStreaming/Private/FastGeoPrimitiveComponent.cpp
@@ -279,7 +279,7 @@ void FFastGeoPrimitiveComponent::InitializeSceneProxyDescDynamicProperties()
 	SceneProxyDesc.CustomPrimitiveData = &CustomPrimitiveData;
 	SceneProxyDesc.Scene = GetScene();
 #if WITH_EDITOR
-	SceneProxyDesc.PrimitiveComponentInterface = GetPrimitiveComponentInterface();
+	SceneProxyDesc.PrimitiveComponentInterface = this;
 #endif
 	SceneProxyDesc.FeatureLevel = SceneProxyDesc.Scene->GetFeatureLevel();
 	TArray<URuntimeVirtualTexture*> const& VirtualTextures = GetRuntimeVirtualTextures();
@@ -345,6 +345,16 @@ void FFastGeoPrimitiveComponent::CreateRenderState(FRegisterComponentContext* Co
 	}
 }
 
+FString FFastGeoPrimitiveComponent::GetName() const
+{
+	return GetNameSafe(GetOwnerContainer());
+}
+
+FString FFastGeoPrimitiveComponent::GetFullName() const
+{
+	return GetFullNameSafe(GetOwnerContainer());
+}
+
 void FFastGeoPrimitiveComponent::DestroyRenderState(FFastGeoDestroyRenderStateContext* Context)
 {
 	FWriteScopeLock WriteLock(*Lock.Get());
@@ -377,7 +387,7 @@ void FFastGeoPrimitiveComponent::DestroyRenderState(FFastGeoDestroyRenderStateCo
 #if WITH_EDITOR
 void FFastGeoPrimitiveComponent::NotifyRenderStateChanged()
 {
-	FObjectCacheEventSink::NotifyRenderStateChanged_Concurrent(GetPrimitiveComponentInterface());
+	FObjectCacheEventSink::NotifyRenderStateChanged_Concurrent(this);
 }
 #endif
 
@@ -473,6 +483,16 @@ void FFastGeoPrimitiveComponent::OnAsyncDestroyPhysicsState()
 	Super::OnAsyncDestroyPhysicsState();
 }
 
+bool FFastGeoPrimitiveComponent::IsRenderStateDirty() const
+{
+	return Super::IsRenderStateDirty();
+}
+
+bool FFastGeoPrimitiveComponent::IsRenderStateCreated() const
+{
+	return Super::IsRenderStateCreated();
+}
+
 bool FFastGeoPrimitiveComponent::ShouldCreateRenderState() const
 {
 	if (!FApp::CanEverRender())
@@ -492,6 +512,27 @@ bool FFastGeoPrimitiveComponent::ShouldCreateRenderState() const
 	return bShouldCreateRenderState;
 }
 
+bool FFastGeoPrimitiveComponent::IsRegistered() const
+{
+	return Super::IsRegistered();
+}
+
+bool FFastGeoPrimitiveComponent::IsUnreachable() const
+{
+	const UObject* OwnerPtr = GetOwnerContainer();
+	return !OwnerPtr || OwnerPtr->IsUnreachable();
+}
+
+UWorld* FFastGeoPrimitiveComponent::GetWorld() const
+{
+	return Super::GetWorld();
+}
+
+FSceneInterface* FFastGeoPrimitiveComponent::GetScene() const
+{
+	return Super::GetScene();
+}
+
 bool FFastGeoPrimitiveComponent::IsFirstPersonRelevant() const
 {
 	return GetSceneProxyDesc().IsFirstPersonRelevant();
@@ -507,16 +548,44 @@ FPrimitiveSceneProxy* FFastGeoPrimitiveComponent::GetSceneProxy() const
 	return PrimitiveSceneData.SceneProxy;
 }
 
-const FTransform& FFastGeoPrimitiveComponent::GetTransform() const
+void FFastGeoPrimitiveComponent::GetStreamableRenderAssetInfo(TArray<FStreamingRenderAssetPrimitiveInfo>&) const
+{
+}
+
+void FFastGeoPrimitiveComponent::MarkRenderStateDirty()
+{
+	Super::MarkRenderStateDirty();
+}
+
+void FFastGeoPrimitiveComponent::DestroyRenderState()
+{
+	DestroyRenderState(nullptr);
+}
+
+FTransform FFastGeoPrimitiveComponent::GetTransform() const
 {
 	return WorldTransform;
 }
 
-const FBoxSphereBounds& FFastGeoPrimitiveComponent::GetBounds() const
+FBoxSphereBounds FFastGeoPrimitiveComponent::GetBounds() const
 {
 	return WorldBounds;
 }
 
+void FFastGeoPrimitiveComponent::GetPrimitiveStats(FPrimitiveStats& PrimitiveStats) const
+{
+}
+
+UObject* FFastGeoPrimitiveComponent::GetUObject()
+{
+	return GetOwnerContainer();
+}
+
+const UObject* FFastGeoPrimitiveComponent::GetUObject() const
+{
+	return GetOwnerContainer();
+}
+
 FMatrix FFastGeoPrimitiveComponent::GetRenderMatrix() const
 {
 	return GetTransform().ToMatrixWithScale();
@@ -682,6 +751,33 @@ void FFastGeoPrimitiveComponent::PrecachePSOs()
 #endif
 }
 
+UObject* FFastGeoPrimitiveComponent::GetOwner() const
+{
+	return GetOwnerContainer();
+}
+
+FString FFastGeoPrimitiveComponent::GetOwnerName() const
+{
+	return GetNameSafe(GetOwnerContainer());
+}
+
+FPrimitiveSceneProxy* FFastGeoPrimitiveComponent::CreateSceneProxy()
+{
+	return CreateSceneProxy(/*OutError*/nullptr);
+}
+
+#if WITH_EDITOR
+HHitProxy* FFastGeoPrimitiveComponent::CreateMeshHitProxy(int32 SectionIndex, int32 MaterialIndex)
+{
+	return nullptr;
+}
+#endif // WITH_EDITOR
+
+HHitProxy* FFastGeoPrimitiveComponent::CreatePrimitiveHitProxies(TArray<TRefCountPtr<HHitProxy>>& OutHitProxies)
+{
+	return nullptr;
+}
+
 bool FFastGeoPrimitiveComponent::IsNavigationRelevant() const
 {
 	if (!bCanEverAffectNavigation)
diff --git a/Engine/Plugins/Experimental/FastGeoStreaming/Source/FastGeoStreaming/Private/FastGeoPrimitiveComponentInterface.cpp b/Engine/Plugins/Experimental/FastGeoStreaming/Source/FastGeoStreaming/Private/FastGeoPrimitiveComponentInterface.cpp
deleted file mode 100644
diff --git a/Engine/Plugins/Experimental/FastGeoStreaming/Source/FastGeoStreaming/Private/FastGeoStaticMeshComponent.cpp b/Engine/Plugins/Experimental/FastGeoStreaming/Source/FastGeoStreaming/Private/FastGeoStaticMeshComponent.cpp
--- a/Engine/Plugins/Experimental/FastGeoStreaming/Source/FastGeoStreaming/Private/FastGeoStaticMeshComponent.cpp
+++ b/Engine/Plugins/Experimental/FastGeoStreaming/Source/FastGeoStreaming/Private/FastGeoStaticMeshComponent.cpp
@@ -93,7 +93,7 @@ void FFastGeoStaticMeshComponentBase::NotifyRenderStateChanged()
 {
 	Super::NotifyRenderStateChanged();
 
-	FObjectCacheEventSink::NotifyStaticMeshChanged_Concurrent(GetStaticMeshComponentInterface());
+	FObjectCacheEventSink::NotifyStaticMeshChanged_Concurrent(this);
 }
 #endif
 
@@ -327,6 +327,11 @@ UStaticMesh* FFastGeoStaticMeshComponentBase::GetStaticMesh() const
 	return GetStaticMeshSceneProxyDesc().StaticMesh;
 }
 
+IPrimitiveComponent* FFastGeoStaticMeshComponentBase::GetPrimitiveComponentInterface()
+{
+	return this;
+}
+
 UObject const* FFastGeoStaticMeshComponentBase::AdditionalStatObject() const
 {
 	return GetStaticMesh();
diff --git a/Engine/Plugins/Experimental/FastGeoStreaming/Source/FastGeoStreaming/Private/FastGeoStaticMeshComponentInterface.cpp b/Engine/Plugins/Experimental/FastGeoStreaming/Source/FastGeoStreaming/Private/FastGeoStaticMeshComponentInterface.cpp
deleted file mode 100644
diff --git a/Engine/Plugins/Experimental/FastGeoStreaming/Source/FastGeoStreaming/Private/FastGeoStreamingModule.cpp b/Engine/Plugins/Experimental/FastGeoStreaming/Source/FastGeoStreaming/Private/FastGeoStreamingModule.cpp
--- a/Engine/Plugins/Experimental/FastGeoStreaming/Source/FastGeoStreaming/Private/FastGeoStreamingModule.cpp
+++ b/Engine/Plugins/Experimental/FastGeoStreaming/Source/FastGeoStreaming/Private/FastGeoStreamingModule.cpp
@@ -20,12 +20,12 @@ FAutoConsoleVariableRef FFastGeoStreamingModule::CVarIsEnabled(
 
 void FFastGeoStreamingModule::StartupModule()
 {
-	IPrimitiveComponent::AddImplementer(FComponentInterfaceImplementation(UFastGeoContainer::StaticClass(),  [](UObject* Obj)
+	IPrimitiveComponent::AddImplementer(TComponentInterfaceImplementation<IPrimitiveComponent>(UFastGeoContainer::StaticClass(),  [](UObject* Obj)
 	{
 		return CastChecked<UFastGeoContainer>(Obj)->GetPrimitiveComponents();
 	}));
 
-	IStaticMeshComponent::AddImplementer(FComponentInterfaceImplementation(UFastGeoContainer::StaticClass(),  [](UObject* Obj)
+	IStaticMeshComponent::AddImplementer(TComponentInterfaceImplementation<IStaticMeshComponent>(UFastGeoContainer::StaticClass(),  [](UObject* Obj)
 	{
 		return CastChecked<UFastGeoContainer>(Obj)->GetStaticMeshComponents();
 	}));
diff --git a/Engine/Source/Runtime/Engine/Classes/Components/ComponentInterfaceIterator.h b/Engine/Source/Runtime/Engine/Classes/Components/ComponentInterfaceIterator.h
--- a/Engine/Source/Runtime/Engine/Classes/Components/ComponentInterfaceIterator.h
+++ b/Engine/Source/Runtime/Engine/Classes/Components/ComponentInterfaceIterator.h
@@ -20,7 +20,7 @@ public:
 	{
 		TArray<UObject*>	ObjectArray;
 
-		for (FComponentInterfaceImplementation& Implementation : T::Implementers)
+		for (TComponentInterfaceImplementation<T>& Implementation : T::Implementers)
 		{
 			if (Implementation.Class)
 			{
@@ -31,7 +31,7 @@ public:
 					Interfaces.Reserve(ObjectArray.Num() + Interfaces.Num());
 					for (UObject* Object : ObjectArray)
 					{
-						T* Interface = reinterpret_cast<T*>(Implementation.Resolver(Object));
+						T* Interface = Implementation.Resolver(Object);
 						check(Interface);
 						Interfaces.Add(Interface);
 					}
@@ -43,11 +43,11 @@ public:
 					GetObjectsOfClass(Implementation.Class, ObjectArray, bIncludeDerivedClasses, AdditionalExclusionFlags, GetObjectIteratorDefaultInternalExclusionFlags(InInternalExclusionFlags));
 					for (UObject* Object : ObjectArray)
 					{
-						TArray<void*> ContainerInterfaces = Implementation.CollectionResolver(Object);
+						TArray<T*> ContainerInterfaces = Implementation.CollectionResolver(Object);
 						Interfaces.Reserve(ContainerInterfaces.Num() + Interfaces.Num());
-						for (void* ContainerInterface : ContainerInterfaces)
+						for (T* ContainerInterface : ContainerInterfaces)
 						{
-							T* Interface = reinterpret_cast<T*>(ContainerInterface);
+							T* Interface = ContainerInterface;
 							check(Interface);
 							Interfaces.Add(Interface);
 						}
diff --git a/Engine/Source/Runtime/Engine/Classes/Components/ComponentInterfaces.h b/Engine/Source/Runtime/Engine/Classes/Components/ComponentInterfaces.h
--- a/Engine/Source/Runtime/Engine/Classes/Components/ComponentInterfaces.h
+++ b/Engine/Source/Runtime/Engine/Classes/Components/ComponentInterfaces.h
@@ -50,22 +50,24 @@ struct FPrimitiveStats
 	TArray<FPrimitiveLODStats> LODStats;
 };
 
-struct FComponentInterfaceImplementation
+template<typename T>
+struct TComponentInterfaceImplementation
 {
-	FComponentInterfaceImplementation(UClass* InClass, TFunction<void* (UObject*)> InResolver)
+	TComponentInterfaceImplementation(UClass* InClass, TFunction<T* (UObject*)> InResolver)
 		: Class(InClass)
 		, Resolver(InResolver)
 	{}
 
-	FComponentInterfaceImplementation(UClass* InClass, TFunction<TArray<void*>(UObject*)> InCollectionResolver)
+	TComponentInterfaceImplementation(UClass* InClass, TFunction<TArray<T*>(UObject*)> InCollectionResolver)
 		: Class(InClass)
 		, CollectionResolver(InCollectionResolver)
 	{}
 
 	UClass*	Class = nullptr;
-	TFunction<void*(UObject*)> Resolver;
-	TFunction<TArray<void*> (UObject*)> CollectionResolver;
+	TFunction<T*(UObject*)> Resolver;
+	TFunction<TArray<T*> (UObject*)> CollectionResolver;
 };
+typedef TComponentInterfaceImplementation<void> FComponentInterfaceImplementation;
 
 class IPrimitiveComponent
 {
@@ -113,14 +115,18 @@ public:
 #endif
 	virtual HHitProxy* CreatePrimitiveHitProxies(TArray<TRefCountPtr<HHitProxy> >& OutHitProxies) = 0;
 	
-	ENGINE_API static void AddImplementer(const FComponentInterfaceImplementation& Implementer);
+	ENGINE_API static void AddImplementer(const TComponentInterfaceImplementation<IPrimitiveComponent>& Implementer);
 	ENGINE_API static void RemoveImplementer(const UClass* ImplementerClass);
+
+	UE_DEPRECATED(5.8, "Use templated version")
+	ENGINE_API static void AddImplementer(const FComponentInterfaceImplementation& Implementer);
+
 protected:
 
 	template<class T> 
 	friend class TComponentInterfaceIterator;
 
-	ENGINE_API static TArray<FComponentInterfaceImplementation> Implementers;
+	ENGINE_API static TArray<TComponentInterfaceImplementation<IPrimitiveComponent>> Implementers;
 };
 
 class IStaticMeshComponent
@@ -140,15 +146,18 @@ public:
 		return (const_cast<IStaticMeshComponent*>(this))->GetPrimitiveComponentInterface();
 	}
 	
-	ENGINE_API static void AddImplementer(const FComponentInterfaceImplementation& Implementer);
+	ENGINE_API static void AddImplementer(const TComponentInterfaceImplementation<IStaticMeshComponent>& Implementer);
 	ENGINE_API static void RemoveImplementer(const UClass* ImplementerClass);
 
+	UE_DEPRECATED(5.8, "Use templated version")
+	ENGINE_API static void AddImplementer(const FComponentInterfaceImplementation& Implementer);
+
 protected:
 
 	template<class T> 
 	friend class TComponentInterfaceIterator;
 
-	ENGINE_API static TArray<FComponentInterfaceImplementation> Implementers;
+	ENGINE_API static TArray<TComponentInterfaceImplementation<IStaticMeshComponent>> Implementers;
 	
 };
 
@@ -168,14 +177,16 @@ public:
 		return (const_cast<ISkinnedMeshComponent*>(this))->GetPrimitiveComponentInterface();
 	}
 	
-	ENGINE_API static void AddImplementer(const FComponentInterfaceImplementation& Implementer);
+	ENGINE_API static void AddImplementer(const TComponentInterfaceImplementation<ISkinnedMeshComponent>& Implementer);
 	ENGINE_API static void RemoveImplementer(const UClass* ImplementerClass);
 
+	UE_DEPRECATED(5.8, "Use templated version")
+	ENGINE_API static void AddImplementer(const FComponentInterfaceImplementation& Implementer);
 protected:
 	template<class T> 
 	friend class TComponentInterfaceIterator;
 
-	ENGINE_API static TArray<FComponentInterfaceImplementation> Implementers;
+	ENGINE_API static TArray<TComponentInterfaceImplementation<ISkinnedMeshComponent>> Implementers;
 };
 
 #pragma region HelperMacros
diff --git a/Engine/Source/Runtime/Engine/Private/Components/ComponentInterfaces.cpp b/Engine/Source/Runtime/Engine/Private/Components/ComponentInterfaces.cpp
--- a/Engine/Source/Runtime/Engine/Private/Components/ComponentInterfaces.cpp
+++ b/Engine/Source/Runtime/Engine/Private/Components/ComponentInterfaces.cpp
@@ -5,16 +5,21 @@
 #include "HAL/IConsoleManager.h"
 #include "UnrealEngine.h"
 
-TArray<FComponentInterfaceImplementation> IPrimitiveComponent::Implementers;
+TArray<TComponentInterfaceImplementation<IPrimitiveComponent>> IPrimitiveComponent::Implementers;
 
-void IPrimitiveComponent::AddImplementer(const FComponentInterfaceImplementation& Implementer)
-{	
+void IPrimitiveComponent::AddImplementer(const TComponentInterfaceImplementation<IPrimitiveComponent>& Implementer)
+{
 	Implementers.Add(Implementer);
 }
 
+void IPrimitiveComponent::AddImplementer(const FComponentInterfaceImplementation& Implementer)
+{
+	AddImplementer(reinterpret_cast<const TComponentInterfaceImplementation<IPrimitiveComponent>&>(Implementer));
+}
+
 void IPrimitiveComponent::RemoveImplementer(const UClass* ImplementerClass)
 {
-	for (TArray<FComponentInterfaceImplementation>::TIterator ImplementerIt(Implementers); ImplementerIt; ++ImplementerIt)
+	for (TArray<TComponentInterfaceImplementation<IPrimitiveComponent>>::TIterator ImplementerIt(Implementers); ImplementerIt; ++ImplementerIt)
 	{
 		if (ImplementerIt->Class == ImplementerClass)
 		{
@@ -23,16 +28,21 @@ void IPrimitiveComponent::RemoveImplementer(const UClass* ImplementerClass)
 	}
 }
 
-TArray<FComponentInterfaceImplementation> IStaticMeshComponent::Implementers;
+TArray<TComponentInterfaceImplementation<IStaticMeshComponent>> IStaticMeshComponent::Implementers;
 
-void IStaticMeshComponent::AddImplementer(const FComponentInterfaceImplementation& Implementer)
+void IStaticMeshComponent::AddImplementer(const TComponentInterfaceImplementation<IStaticMeshComponent>& Implementer)
 {
 	Implementers.Add(Implementer);
 }
 
+void IStaticMeshComponent::AddImplementer(const FComponentInterfaceImplementation& Implementer)
+{
+	AddImplementer(reinterpret_cast<const TComponentInterfaceImplementation<IStaticMeshComponent>&>(Implementer));
+}
+
 void IStaticMeshComponent::RemoveImplementer(const UClass* ImplementerClass)
 {
-	for (TArray<FComponentInterfaceImplementation>::TIterator ImplementerIt(Implementers); ImplementerIt; ++ImplementerIt)
+	for (TArray<TComponentInterfaceImplementation<IStaticMeshComponent>>::TIterator ImplementerIt(Implementers); ImplementerIt; ++ImplementerIt)
 	{
 		if (ImplementerIt->Class == ImplementerClass)
 		{
@@ -41,16 +51,21 @@ void IStaticMeshComponent::RemoveImplementer(const UClass* ImplementerClass)
 	}
 }
 
-TArray<FComponentInterfaceImplementation> ISkinnedMeshComponent::Implementers;
+TArray<TComponentInterfaceImplementation<ISkinnedMeshComponent>> ISkinnedMeshComponent::Implementers;
 
-void ISkinnedMeshComponent::AddImplementer(const FComponentInterfaceImplementation& Implementer)
+void ISkinnedMeshComponent::AddImplementer(const TComponentInterfaceImplementation<ISkinnedMeshComponent>& Implementer)
 {
 	Implementers.Add(Implementer);
 }
 
+void ISkinnedMeshComponent::AddImplementer(const FComponentInterfaceImplementation& Implementer)
+{
+	AddImplementer(reinterpret_cast<const TComponentInterfaceImplementation<ISkinnedMeshComponent>&>(Implementer));
+}
+
 void ISkinnedMeshComponent::RemoveImplementer(const UClass* ImplementerClass)
 {
-	for (TArray<FComponentInterfaceImplementation>::TIterator ImplementerIt(Implementers); ImplementerIt; ++ImplementerIt)
+	for (TArray<TComponentInterfaceImplementation<ISkinnedMeshComponent>>::TIterator ImplementerIt(Implementers); ImplementerIt; ++ImplementerIt)
 	{
 		if (ImplementerIt->Class == ImplementerClass)
 		{
diff --git a/Engine/Source/Runtime/Engine/Private/UnrealEngine.cpp b/Engine/Source/Runtime/Engine/Private/UnrealEngine.cpp
--- a/Engine/Source/Runtime/Engine/Private/UnrealEngine.cpp
+++ b/Engine/Source/Runtime/Engine/Private/UnrealEngine.cpp
@@ -339,9 +339,9 @@ void FEngineModule::StartupModule()
 		return true;
 	});
 
-	IPrimitiveComponent::AddImplementer({UPrimitiveComponent::StaticClass(),  [](UObject* Obj){return Cast<UPrimitiveComponent>(Obj)->GetPrimitiveComponentInterface();}});
-	IStaticMeshComponent::AddImplementer({UStaticMeshComponent::StaticClass(), [](UObject* Obj){return Cast<UStaticMeshComponent>(Obj)->GetStaticMeshComponentInterface();}});
-	ISkinnedMeshComponent::AddImplementer({USkinnedMeshComponent::StaticClass(), [](UObject* Obj) {return Cast<USkinnedMeshComponent>(Obj)->GetSkinnedMeshComponentInterface();}});
+	IPrimitiveComponent::AddImplementer(TComponentInterfaceImplementation<IPrimitiveComponent>{UPrimitiveComponent::StaticClass(),  [](UObject* Obj){return Cast<UPrimitiveComponent>(Obj)->GetPrimitiveComponentInterface();}});
+	IStaticMeshComponent::AddImplementer(TComponentInterfaceImplementation<IStaticMeshComponent>{UStaticMeshComponent::StaticClass(), [](UObject* Obj){return Cast<UStaticMeshComponent>(Obj)->GetStaticMeshComponentInterface();}});
+	ISkinnedMeshComponent::AddImplementer(TComponentInterfaceImplementation<ISkinnedMeshComponent>{USkinnedMeshComponent::StaticClass(), [](UObject* Obj) {return Cast<USkinnedMeshComponent>(Obj)->GetSkinnedMeshComponentInterface();}});
 
 #if OBJECT_TRACE_ENABLED
 	FObjectTrace::Init();
