[FastGeo] Added UFastGeoContainer::CreateTransient and UFastGeoContainer::DestroyTransient for runtime spawning and asynchronous destruction of FastGeo content.
During world cleanup, any pending destruction of transient containers will wait for asynchronous tasks to complete.

diff --git a/Engine/Plugins/Experimental/FastGeoStreaming/Source/FastGeoStreaming/Internal/FastGeoContainer.h b/Engine/Plugins/Experimental/FastGeoStreaming/Source/FastGeoStreaming/Internal/FastGeoContainer.h
--- a/Engine/Plugins/Experimental/FastGeoStreaming/Source/FastGeoStreaming/Internal/FastGeoContainer.h
+++ b/Engine/Plugins/Experimental/FastGeoStreaming/Source/FastGeoStreaming/Internal/FastGeoContainer.h
@@ -104,6 +104,17 @@ class FASTGEOSTREAMING_API UFastGeoContainer : public UAssetUserData, public IAs
 	GENERATED_BODY()
 
 public:
+
+	/**
+	 * Creates a transient FastGeoContainer for a specific world (used for runtime).
+	 */
+	static UFastGeoContainer* CreateTransient(UWorld* InWorld);
+
+	/**
+	 * Destroys a transient FastGeoContainer.
+	 */
+	static void DestroyTransient(UFastGeoContainer* InFastGeo);
+
 	void Register();
 	void Unregister();
 	bool IsRegistered() const;
@@ -292,10 +303,12 @@ private:
 	TArray<FNavigationElementHandle> NavigationElementHandles;
 
 	bool bIsRegistered = false;
+	bool bIsBeingAsyncDestroyed = false;
 	int32 RegistrationEpoch = 0;
 
 	FRenderCommandFence DestroyFence;
 
 	friend class FFastGeoGatherFastGeoContainerAssetRefsArchive;
 	friend struct FFastGeoRegisteredComponent;
+	friend class UFastGeoWorldSubsystem;
 };
\ No newline at end of file
diff --git a/Engine/Plugins/Experimental/FastGeoStreaming/Source/FastGeoStreaming/Private/FastGeoComponent.cpp b/Engine/Plugins/Experimental/FastGeoStreaming/Source/FastGeoStreaming/Private/FastGeoComponent.cpp
--- a/Engine/Plugins/Experimental/FastGeoStreaming/Source/FastGeoStreaming/Private/FastGeoComponent.cpp
+++ b/Engine/Plugins/Experimental/FastGeoStreaming/Source/FastGeoStreaming/Private/FastGeoComponent.cpp
@@ -139,10 +139,7 @@ void FFastGeoComponent::MarkRenderStateDirty()
 	{
 		bRenderStateDirty = true;
 
-		UFastGeoWorldSubsystem* WorldSubsystem = GetWorld()->GetSubsystem<UFastGeoWorldSubsystem>();
-		if (ensure(WorldSubsystem))
-		{
-			WorldSubsystem->AddToComponentsPendingRecreate(this);
-		}
+		UFastGeoWorldSubsystem* WorldSubsystem = GetWorld()->GetSubsystemChecked<UFastGeoWorldSubsystem>();
+		WorldSubsystem->AddToComponentsPendingRecreate(this);
 	}
 }
\ No newline at end of file
diff --git a/Engine/Plugins/Experimental/FastGeoStreaming/Source/FastGeoStreaming/Private/FastGeoContainer.cpp b/Engine/Plugins/Experimental/FastGeoStreaming/Source/FastGeoStreaming/Private/FastGeoContainer.cpp
--- a/Engine/Plugins/Experimental/FastGeoStreaming/Source/FastGeoStreaming/Private/FastGeoContainer.cpp
+++ b/Engine/Plugins/Experimental/FastGeoStreaming/Source/FastGeoStreaming/Private/FastGeoContainer.cpp
@@ -175,6 +175,12 @@ void UFastGeoContainer::Register()
 		return;
 	}
 
+	if (bIsBeingAsyncDestroyed)
+	{
+		UE_LOG(LogFastGeoStreaming, Warning, TEXT("UFastGeoContainer::Register() ignored because FastGeo is being destroyed."));
+		return;
+	}
+
 	if (!bIsRegistered)
 	{
 		++RegistrationEpoch;
@@ -205,7 +211,7 @@ void UFastGeoContainer::Register()
 			});
 		});
 
-		UFastGeoWorldSubsystem* WorldSubsystem = GetWorld()->GetSubsystem<UFastGeoWorldSubsystem>();
+		UFastGeoWorldSubsystem* WorldSubsystem = GetWorld()->GetSubsystemChecked<UFastGeoWorldSubsystem>();
 		if (!PendingCreate.RenderState.ComponentsToProcess.IsEmpty())
 		{
 			WorldSubsystem->PushAsyncCreateRenderStateJob(this);
@@ -237,6 +243,62 @@ bool UFastGeoContainer::IsRegistered() const
 	return bIsRegistered;
 }
 
+UFastGeoContainer* UFastGeoContainer::CreateTransient(UWorld* InWorld)
+{
+	check(IsInGameThread());
+
+	if (!IsValid(InWorld))
+	{
+		UE_LOG(LogFastGeoStreaming, Error, TEXT("UFastGeoContainer::CreateTransient() called with an invalid world."));
+		return nullptr;
+	}
+
+	ULevel* PersistentLevel = InWorld->PersistentLevel;
+	check(PersistentLevel && PersistentLevel->IsPersistentLevel());
+
+	return NewObject<UFastGeoContainer>(PersistentLevel, NAME_None, RF_Transient);
+}
+
+void UFastGeoContainer::DestroyTransient(UFastGeoContainer* InFastGeo)
+{
+	check(IsInGameThread());
+
+	if (!IsValid(InFastGeo))
+	{
+		UE_LOG(LogFastGeoStreaming, Error, TEXT("UFastGeoContainer::DestroyTransient() called on an invalid container."));
+		return;
+	}
+
+	if (!InFastGeo->GetWorld())
+	{
+		UE_LOG(LogFastGeoStreaming, Error, TEXT("UFastGeoContainer::DestroyTransient() called on a container with an invalid world."));
+		return;
+	}
+
+	if (!InFastGeo->HasAnyFlags(RF_Transient))
+	{
+		UE_LOG(LogFastGeoStreaming, Error, TEXT("UFastGeoContainer::DestroyTransient() is meant to be called on a transient container."));
+		return;
+	}
+
+	if (!InFastGeo->GetLevel()->IsPersistentLevel())
+	{
+		UE_LOG(LogFastGeoStreaming, Error, TEXT("UFastGeoContainer::DestroyTransient() expects the transient container to be outered to the persistent level."));
+		return;
+	}
+
+	if (InFastGeo->bIsBeingAsyncDestroyed)
+	{
+		UE_LOG(LogFastGeoStreaming, Error, TEXT("UFastGeoContainer::DestroyTransient() called on a container that is already being destroyed."));
+		return;
+	}
+
+	InFastGeo->bIsBeingAsyncDestroyed = true;
+
+	UFastGeoWorldSubsystem* WorldSubsystem = InFastGeo->GetWorld()->GetSubsystemChecked<UFastGeoWorldSubsystem>();
+	WorldSubsystem->PushAsyncDestroyContainer(InFastGeo);
+}
+
 void UFastGeoContainer::Unregister()
 {
 	TRACE_CPUPROFILER_EVENT_SCOPE(UFastGeoContainer::Unregister);
@@ -264,7 +326,7 @@ void UFastGeoContainer::Unregister()
 
 		UnregisterFromNavigationSystem();
 		
-		UFastGeoWorldSubsystem* WorldSubsystem = GetWorld()->GetSubsystem<UFastGeoWorldSubsystem>();
+		UFastGeoWorldSubsystem* WorldSubsystem = GetWorld()->GetSubsystemChecked<UFastGeoWorldSubsystem>();
 		if (!PendingDestroy.RenderState.ComponentsToProcess.IsEmpty())
 		{
 			WorldSubsystem->PushAsyncDestroyRenderStateJob(this);
@@ -284,7 +346,7 @@ void UFastGeoContainer::Tick(bool bWaitForCompletion)
 {
 	UWorld* World = GetWorld();
 	check(World);
-	UFastGeoWorldSubsystem* WorldSubsystem = World->GetSubsystem<UFastGeoWorldSubsystem>();
+	UFastGeoWorldSubsystem* WorldSubsystem = World->GetSubsystemChecked<UFastGeoWorldSubsystem>();
 
 	do
 	{
@@ -430,7 +492,7 @@ void UFastGeoContainer::OnCreateRenderStateBegin_GameThread()
 	check(PendingCreate.RenderState.NumToProcess == 0);
 	check(PendingCreate.RenderState.TotalNumProcessed != PendingCreate.RenderState.ComponentsToProcess.Num());
 
-	PendingCreate.bIsInBlockingWait = GetWorld()->GetSubsystem<UFastGeoWorldSubsystem>()->IsWaitingForCompletion();
+	PendingCreate.bIsInBlockingWait = GetWorld()->GetSubsystemChecked<UFastGeoWorldSubsystem>()->IsWaitingForCompletion();
 
 	PendingCreate.RenderState.NumToProcess = PendingCreate.RenderState.ComponentsToProcess.Num() - PendingCreate.RenderState.TotalNumProcessed;
 }
@@ -440,7 +502,7 @@ void UFastGeoContainer::OnDestroyRenderStateBegin_GameThread()
 	check(PendingDestroy.RenderState.NumToProcess == 0);
 	check(PendingDestroy.RenderState.TotalNumProcessed != PendingDestroy.RenderState.ComponentsToProcess.Num());
 
-	PendingDestroy.bIsInBlockingWait = GetWorld()->GetSubsystem<UFastGeoWorldSubsystem>()->IsWaitingForCompletion();
+	PendingDestroy.bIsInBlockingWait = GetWorld()->GetSubsystemChecked<UFastGeoWorldSubsystem>()->IsWaitingForCompletion();
 
 	// TODO_FASTGEO: Evaluate if we need throttling of the async destruction task
 	PendingDestroy.RenderState.NumToProcess = PendingDestroy.RenderState.ComponentsToProcess.Num();
@@ -458,7 +520,7 @@ void UFastGeoContainer::OnCreateRenderState_Concurrent()
 	int32 AvailableComponentsBudget;
 	int32 TimeEpoch;
 
-	UFastGeoWorldSubsystem* WorldSubsystem = GetWorld()->GetSubsystem<UFastGeoWorldSubsystem>();
+	UFastGeoWorldSubsystem* WorldSubsystem = GetWorld()->GetSubsystemChecked<UFastGeoWorldSubsystem>();
 	WorldSubsystem->RequestAsyncRenderStateTasksBudget_Concurrent(AvailableTimeBudgetMS, AvailableComponentsBudget, TimeEpoch);
 	
 	const int32 ComponentsBudget = FMath::Min(NumComponentsToProcess, AvailableComponentsBudget);
@@ -535,7 +597,7 @@ void UFastGeoContainer::OnCreateRenderStateEnd_GameThread()
 
 	if (PendingCreate.RenderState.TotalNumProcessed < PendingCreate.RenderState.ComponentsToProcess.Num())
 	{
-		UFastGeoWorldSubsystem* WorldSubsystem = GetWorld()->GetSubsystem<UFastGeoWorldSubsystem>();
+		UFastGeoWorldSubsystem* WorldSubsystem = GetWorld()->GetSubsystemChecked<UFastGeoWorldSubsystem>();
 		WorldSubsystem->PushAsyncCreateRenderStateJob(this);
 	}
 }
@@ -552,7 +614,7 @@ void UFastGeoContainer::OnDestroyRenderStateEnd_GameThread()
 
 	if (PendingDestroy.RenderState.TotalNumProcessed < PendingDestroy.RenderState.ComponentsToProcess.Num())
 	{
-		UFastGeoWorldSubsystem* WorldSubsystem = GetWorld()->GetSubsystem<UFastGeoWorldSubsystem>();
+		UFastGeoWorldSubsystem* WorldSubsystem = GetWorld()->GetSubsystemChecked<UFastGeoWorldSubsystem>();
 		WorldSubsystem->PushAsyncDestroyRenderStateJob(this);
 	}
 }
diff --git a/Engine/Plugins/Experimental/FastGeoStreaming/Source/FastGeoStreaming/Private/FastGeoWorldSubsystem.cpp b/Engine/Plugins/Experimental/FastGeoStreaming/Source/FastGeoStreaming/Private/FastGeoWorldSubsystem.cpp
--- a/Engine/Plugins/Experimental/FastGeoStreaming/Source/FastGeoStreaming/Private/FastGeoWorldSubsystem.cpp
+++ b/Engine/Plugins/Experimental/FastGeoStreaming/Source/FastGeoStreaming/Private/FastGeoWorldSubsystem.cpp
@@ -171,6 +171,9 @@ void UFastGeoWorldSubsystem::PostInitialize()
 	Super::PostInitialize();
 
 	UWorld* World = GetWorld();
+
+	FWorldDelegates::OnWorldCleanup.AddUObject(this, &UFastGeoWorldSubsystem::OnWorldCleanup);
+
 	if (World->IsPartitionedWorld() && World->IsGameWorld())
 	{
 		World->OnAllLevelsChanged().AddUObject(this, &UFastGeoWorldSubsystem::OnUpdateLevelStreaming);
@@ -194,6 +197,8 @@ void UFastGeoWorldSubsystem::Deinitialize()
 {
 	UWorld* World = GetWorld();
 
+	FWorldDelegates::OnWorldCleanup.RemoveAll(this);
+
 	if (World->IsGameWorld())
 	{
 		World->OnAllLevelsChanged().RemoveAll(this);
@@ -268,6 +273,22 @@ void UFastGeoWorldSubsystem::OnLevelStartedRemoveFromWorld(UWorld* World, const
 	}
 }
 
+void UFastGeoWorldSubsystem::OnWorldCleanup(UWorld* World, bool bSessionEnded, bool bCleanupResources)
+{
+	if (World != GetWorld())
+	{
+		return;
+	}
+
+	const bool bWaitForCompletion = true;
+	for (const TObjectPtr<UFastGeoContainer>& Container : PendingDestroyContainers)
+	{
+		Container->Tick(bWaitForCompletion);
+		check(!Container->HasAnyPendingDestroyTasks());
+	}
+	PendingDestroyContainers.Reset();
+}
+
 void UFastGeoWorldSubsystem::OnLevelComponentsUpdated(UWorld* World, ULevel* Level)
 {
 	if (World != GetWorld())
@@ -462,13 +483,25 @@ bool UFastGeoWorldSubsystem::IsWaitingForCompletion() const
 			GetWorld()->IsBeingCleanedUp();
 }
 
+bool UFastGeoWorldSubsystem::IsTickableInEditor() const
+{
+	return true;
+}
+
 void UFastGeoWorldSubsystem::Tick(float DeltaTime)
 {
 	FWriteScopeLock ScopeLock(Lock);
-
 	++TimeEpoch;
 	UsedAsyncRenderStateTasksTimeBudgetMS = 0;
 	UsedNumComponentsToProcessBudget = 0;
+	
+	for (auto It = PendingDestroyContainers.CreateIterator(); It; ++It)
+	{
+		if (!(*It)->HasAnyPendingDestroyTasks())
+		{
+			It.RemoveCurrent();
+		}
+	}
 }
 
 TStatId UFastGeoWorldSubsystem::GetStatId() const
@@ -521,6 +554,19 @@ void UFastGeoWorldSubsystem::CommitAsyncRenderStateTasksBudget_Concurrent(float
 	}
 }
 
+void UFastGeoWorldSubsystem::PushAsyncDestroyContainer(UFastGeoContainer* FastGeo)
+{
+	check(IsInGameThread());
+	check(FastGeo->bIsBeingAsyncDestroyed);
+
+	FastGeo->Unregister();
+
+	if (FastGeo->HasAnyPendingDestroyTasks())
+	{
+		PendingDestroyContainers.Add(FastGeo);
+	}
+}
+
 void UFastGeoWorldSubsystem::PushAsyncCreateRenderStateJob(UFastGeoContainer* FastGeo)
 {
 	check(IsInGameThread());
diff --git a/Engine/Plugins/Experimental/FastGeoStreaming/Source/FastGeoStreaming/Private/FastGeoWorldSubsystem.h b/Engine/Plugins/Experimental/FastGeoStreaming/Source/FastGeoStreaming/Private/FastGeoWorldSubsystem.h
--- a/Engine/Plugins/Experimental/FastGeoStreaming/Source/FastGeoStreaming/Private/FastGeoWorldSubsystem.h
+++ b/Engine/Plugins/Experimental/FastGeoStreaming/Source/FastGeoStreaming/Private/FastGeoWorldSubsystem.h
@@ -33,6 +33,7 @@ public:
 	//~ End UWorldSubsystem interface.
 
 	//~ Begin UTickableWorldSubsystem interface
+	virtual bool IsTickableInEditor() const override;
 	virtual void Tick(float DeltaTime) override;
 	virtual TStatId GetStatId() const override;
 	//~ End UTickableWorldSubsystem interface
@@ -80,6 +81,7 @@ public:
 
 private:
 	void OnUpdateLevelStreaming();
+	void OnWorldCleanup(UWorld* World, bool bSessionEnded, bool bCleanupResources);
 	void OnLevelStreamingStateChanged(UWorld* World, const ULevelStreaming* StreamingLevel, ULevel* LevelIfLoaded, ELevelStreamingState PrevState, ELevelStreamingState NewState);
 	void OnLevelStartedAddToWorld(UWorld* World, const ULevelStreaming* StreamingLevel, ULevel* Level);
 	void OnLevelStartedRemoveFromWorld(UWorld* World, const ULevelStreaming* StreamingLevel, ULevel* Level);
@@ -98,6 +100,11 @@ private:
 	void RequestAsyncRenderStateTasksBudget_Concurrent(float& OutAvailableTimeBudgetMS, int32& OutAvaiableComponentsBudget, int32& OutTimeEpoch);
 	void CommitAsyncRenderStateTasksBudget_Concurrent(float InUsedTimeBudgetMS, int32& InUsedComponentsBudget, int32 TimeEpoch);
 
+	void PushAsyncDestroyContainer(UFastGeoContainer* FastGeo);
+
+	UPROPERTY()
+	TSet<TObjectPtr<UFastGeoContainer>> PendingDestroyContainers;
+
 	FDelegateHandle Handle_OnLevelStreamingStateChanged;
 	FDelegateHandle Handle_OnLevelBeginAddToWorld;
 	FDelegateHandle Handle_OnLevelBeginRemoveFromWorld;
diff --git a/Engine/Plugins/Experimental/PCGInterops/PCGFastGeoInterop/Source/PCGFastGeoInterop/Private/Components/PCGManagedFastGeoContainer.cpp b/Engine/Plugins/Experimental/PCGInterops/PCGFastGeoInterop/Source/PCGFastGeoInterop/Private/Components/PCGManagedFastGeoContainer.cpp
--- a/Engine/Plugins/Experimental/PCGInterops/PCGFastGeoInterop/Source/PCGFastGeoInterop/Private/Components/PCGManagedFastGeoContainer.cpp
+++ b/Engine/Plugins/Experimental/PCGInterops/PCGFastGeoInterop/Source/PCGFastGeoInterop/Private/Components/PCGManagedFastGeoContainer.cpp
@@ -10,11 +10,7 @@ bool UPCGManagedFastGeoContainer::Release(bool bHardRelease, TSet<TSoftObjectPtr
 
 	if (FastGeo)
 	{
-		FastGeo->Unregister();
-
-		// Note: This may be costly. Monitor performance and possibly build a way to keep the FastGeo container alive until async work is done.
-		FastGeo->Tick(/*bWaitForCompletion=*/true);
-
+		UFastGeoContainer::DestroyTransient(FastGeo);
 		FastGeo = nullptr;
 	}
 
diff --git a/Engine/Plugins/Experimental/PCGInterops/PCGFastGeoInterop/Source/PCGFastGeoInterop/Private/Compute/PrimitiveFactories/PCGPrimitiveFactoryFastGeoPISMC.cpp b/Engine/Plugins/Experimental/PCGInterops/PCGFastGeoInterop/Source/PCGFastGeoInterop/Private/Compute/PrimitiveFactories/PCGPrimitiveFactoryFastGeoPISMC.cpp
--- a/Engine/Plugins/Experimental/PCGInterops/PCGFastGeoInterop/Source/PCGFastGeoInterop/Private/Compute/PrimitiveFactories/PCGPrimitiveFactoryFastGeoPISMC.cpp
+++ b/Engine/Plugins/Experimental/PCGInterops/PCGFastGeoInterop/Source/PCGFastGeoInterop/Private/Compute/PrimitiveFactories/PCGPrimitiveFactoryFastGeoPISMC.cpp
@@ -90,17 +90,13 @@ bool FPCGPrimitiveFactoryFastGeoPISMC::Create(FPCGContext* InContext)
 	}
 
 	UWorld* World = InContext->ExecutionSource->GetExecutionState().GetWorld();
-	ULevel* Level = World ? World->PersistentLevel.Get() : nullptr;
-
-	if (!Level)
+	UFastGeoContainer* FastGeo = UFastGeoContainer::CreateTransient(World);
+	if (!FastGeo)
 	{
-		UE_LOG(LogPCG, Error, TEXT("No persistent level found, geometry will not be created."));
+		UE_LOG(LogPCG, Error, TEXT("Failed to create a transient FastGeoContainer, geometry will not be created."));
 		return true;
 	}
 
-	// Must be outer'd to level - so pick world persistent level.
-	UFastGeoContainer* FastGeo = FPCGContext::NewObject_AnyThread<UFastGeoContainer>(InContext, Level);
-
 	const FString DebugName = ExecutionSource->GetExecutionState().GetDebugName();
 	TUniquePtr<FFastGeoComponentCluster> ComponentCluster = MakeUnique<FFastGeoComponentCluster>(FastGeo, *FString::Printf(TEXT("FastGeoComponentCluster_%s"), *DebugName));
 
