[FastGeo] Reduce scope of write lock in UFastGeoWorldSubsystem::Tick.

diff --git a/Engine/Plugins/Experimental/FastGeoStreaming/Source/FastGeoStreaming/Private/FastGeoWorldSubsystem.cpp b/Engine/Plugins/Experimental/FastGeoStreaming/Source/FastGeoStreaming/Private/FastGeoWorldSubsystem.cpp
--- a/Engine/Plugins/Experimental/FastGeoStreaming/Source/FastGeoStreaming/Private/FastGeoWorldSubsystem.cpp
+++ b/Engine/Plugins/Experimental/FastGeoStreaming/Source/FastGeoStreaming/Private/FastGeoWorldSubsystem.cpp
@@ -490,10 +490,12 @@ bool UFastGeoWorldSubsystem::IsTickableInEditor() const
 
 void UFastGeoWorldSubsystem::Tick(float DeltaTime)
 {
-	FWriteScopeLock ScopeLock(Lock);
-	++TimeEpoch;
-	UsedAsyncRenderStateTasksTimeBudgetMS = 0;
-	UsedNumComponentsToProcessBudget = 0;
+	{
+		FWriteScopeLock ScopeLock(Lock);
+		++TimeEpoch;
+		UsedAsyncRenderStateTasksTimeBudgetMS = 0;
+		UsedNumComponentsToProcessBudget = 0;
+	}
 	
 	for (auto It = PendingDestroyContainers.CreateIterator(); It; ++It)
 	{
