[FastGeo] Added FastGeo support for MeshPaint textures in StaticMeshComponent.

diff --git a/Engine/Plugins/Experimental/FastGeoStreaming/Source/FastGeoStreaming/Internal/FastGeoContainer.h b/Engine/Plugins/Experimental/FastGeoStreaming/Source/FastGeoStreaming/Internal/FastGeoContainer.h
--- a/Engine/Plugins/Experimental/FastGeoStreaming/Source/FastGeoStreaming/Internal/FastGeoContainer.h
+++ b/Engine/Plugins/Experimental/FastGeoStreaming/Source/FastGeoStreaming/Internal/FastGeoContainer.h
@@ -196,6 +196,9 @@ public:
 		return ForEachComponentClusterBreakable<UFastGeoContainer, TComponentCluster, TFunc>(this, Forward<TFunc>(InFunc));
 	}
 
+	DECLARE_MULTICAST_DELEGATE(FOnUnregistered);
+	FOnUnregistered& GetOnUnregistered() { return OnUnregistered; }
+
 protected:
 
 	FFastGeoComponentClusterPendingState PendingCreate;
@@ -308,6 +311,8 @@ private:
 
 	FRenderCommandFence DestroyFence;
 
+	FOnUnregistered OnUnregistered;
+
 	friend class FFastGeoGatherFastGeoContainerAssetRefsArchive;
 	friend struct FFastGeoRegisteredComponent;
 	friend class UFastGeoWorldSubsystem;
diff --git a/Engine/Plugins/Experimental/FastGeoStreaming/Source/FastGeoStreaming/Internal/FastGeoStaticMeshComponent.h b/Engine/Plugins/Experimental/FastGeoStreaming/Source/FastGeoStreaming/Internal/FastGeoStaticMeshComponent.h
--- a/Engine/Plugins/Experimental/FastGeoStreaming/Source/FastGeoStreaming/Internal/FastGeoStaticMeshComponent.h
+++ b/Engine/Plugins/Experimental/FastGeoStreaming/Source/FastGeoStreaming/Internal/FastGeoStaticMeshComponent.h
@@ -117,6 +117,8 @@ protected:
 		FColorVertexBuffer* OverrideVertexColors;
 	};
 	TArray<FFastGeoStaticMeshComponentLODInfo> LODData;
+	FDelegateHandle OnResourceUpdated;
+	FDelegateHandle OnContainerUnregistered;
 
 private:
 	friend class FMeshComponentHelper;
diff --git a/Engine/Plugins/Experimental/FastGeoStreaming/Source/FastGeoStreaming/Private/FastGeoComponentCluster.cpp b/Engine/Plugins/Experimental/FastGeoStreaming/Source/FastGeoStreaming/Private/FastGeoComponentCluster.cpp
--- a/Engine/Plugins/Experimental/FastGeoStreaming/Source/FastGeoStreaming/Private/FastGeoComponentCluster.cpp
+++ b/Engine/Plugins/Experimental/FastGeoStreaming/Source/FastGeoStreaming/Private/FastGeoComponentCluster.cpp
@@ -259,6 +259,11 @@ FFastGeoComponent& FFastGeoComponentCluster::AddComponent(FFastGeoElementType In
 		unimplemented(); // Should never be reached, missing type handling
 	}
 
+	if (NewComponent)
+	{
+		NewComponent->SetOwnerComponentCluster(this);
+	}
+
 	return *NewComponent;
 }
 
diff --git a/Engine/Plugins/Experimental/FastGeoStreaming/Source/FastGeoStreaming/Private/FastGeoContainer.cpp b/Engine/Plugins/Experimental/FastGeoStreaming/Source/FastGeoStreaming/Private/FastGeoContainer.cpp
--- a/Engine/Plugins/Experimental/FastGeoStreaming/Source/FastGeoStreaming/Private/FastGeoContainer.cpp
+++ b/Engine/Plugins/Experimental/FastGeoStreaming/Source/FastGeoStreaming/Private/FastGeoContainer.cpp
@@ -12,6 +12,7 @@
 #include "Async/ParallelFor.h"
 #include "Engine/Level.h"
 #include "Engine/SkeletalMesh.h"
+#include "Engine/Texture.h"
 #include "Engine/World.h"
 #include "Misc/Timeout.h"
 #include "NavigationSystem.h"
@@ -334,6 +335,8 @@ void UFastGeoContainer::Unregister()
 
 		WorldSubsystem->PushAsyncDestroyPhysicsStateJobs(this);
 		bIsRegistered = false;
+
+		OnUnregistered.Broadcast();
 	}
 
 	if (PendingDestroy.HasAnyPendingState())
@@ -679,6 +682,7 @@ class FFastGeoGatherFastGeoContainerAssetRefsArchive : public FArchive
 public:
 	FFastGeoGatherFastGeoContainerAssetRefsArchive(UFastGeoContainer& Container)
 	{
+		KnownClasses.Add(UTexture::StaticClass());
 		KnownClasses.Add(URuntimeVirtualTexture::StaticClass());
 		KnownClasses.Add(UStaticMesh::StaticClass());
 		KnownClasses.Add(UMaterialInterface::StaticClass());
diff --git a/Engine/Plugins/Experimental/FastGeoStreaming/Source/FastGeoStreaming/Private/FastGeoStaticMeshComponent.cpp b/Engine/Plugins/Experimental/FastGeoStreaming/Source/FastGeoStreaming/Private/FastGeoStaticMeshComponent.cpp
--- a/Engine/Plugins/Experimental/FastGeoStreaming/Source/FastGeoStreaming/Private/FastGeoStaticMeshComponent.cpp
+++ b/Engine/Plugins/Experimental/FastGeoStreaming/Source/FastGeoStreaming/Private/FastGeoStaticMeshComponent.cpp
@@ -5,6 +5,7 @@
 #include "Engine/Engine.h"
 #include "FastGeoContainer.h"
 #include "FastGeoLog.h"
+#include "FastGeoWeakElement.h"
 #include "AI/Navigation/NavCollisionBase.h"
 #include "AI/Navigation/NavigationRelevantData.h"
 #include "Materials/MaterialRelevance.h"
@@ -15,6 +16,7 @@
 #include "PhysicalMaterials/PhysicalMaterial.h"
 #include "NaniteVertexFactory.h"
 #include "StaticMeshSceneProxy.h"
+#include "VT/MeshPaintVirtualTexture.h"
 
 #if WITH_EDITOR
 #include "ObjectCacheEventSink.h"
@@ -35,6 +37,7 @@ void FFastGeoStaticMeshComponentBase::Serialize(FArchive& Ar)
 	// Serialize persistent data from FStaticMeshSceneProxyDesc
 	FStaticMeshSceneProxyDesc& SceneProxyDesc = GetStaticMeshSceneProxyDesc();
 	Ar << SceneProxyDesc.StaticMesh;
+	Ar << SceneProxyDesc.MeshPaintTexture;
 	Ar << SceneProxyDesc.OverlayMaterial;
 	Ar << SceneProxyDesc.MaterialSlotsOverlayMaterial;
 	Ar << SceneProxyDesc.OverlayMaterialMaxDrawDistance;
@@ -45,6 +48,7 @@ void FFastGeoStaticMeshComponentBase::Serialize(FArchive& Ar)
 	Ar << SceneProxyDesc.DistanceFieldSelfShadowBias;
 	Ar << SceneProxyDesc.DistanceFieldIndirectShadowMinVisibility;
 	Ar << SceneProxyDesc.StaticLightMapResolution;
+	Ar << SceneProxyDesc.MeshPaintTextureCoordinateIndex;
 	FArchive_Serialize_BitfieldBool(Ar, SceneProxyDesc.bReverseCulling);
 	FArchive_Serialize_BitfieldBool(Ar, SceneProxyDesc.bEvaluateWorldPositionOffset);
 	FArchive_Serialize_BitfieldBool(Ar, SceneProxyDesc.bOverrideMinLOD);
@@ -71,6 +75,48 @@ void FFastGeoStaticMeshComponentBase::InitializeDynamicProperties()
 	}
 #endif
 
+#if WITH_EDITOR
+	check(GetWorld());
+	if (GetWorld()->IsGameWorld())
+#endif
+	{
+		if (UMeshPaintVirtualTexture* MeshPaintTexture = Cast<UMeshPaintVirtualTexture>(GetStaticMeshSceneProxyDesc().MeshPaintTexture))
+		{
+			FWeakFastGeoComponent WeakComponent(this);
+			MeshPaintTexture->GetOnResourceUpdated().Remove(OnResourceUpdated);
+			OnResourceUpdated = MeshPaintTexture->GetOnResourceUpdated().AddLambda([WeakComponent]()
+			{
+				if (FFastGeoStaticMeshComponentBase* Component = WeakComponent.Get<FFastGeoStaticMeshComponentBase>())
+				{
+					Component->MarkRenderStateDirty();
+				}
+			});
+
+			// Attach to container unregistration
+			UFastGeoContainer* Container = GetOwnerContainer();
+			Container->GetOnUnregistered().Remove(OnContainerUnregistered);
+			OnContainerUnregistered = Container->GetOnUnregistered().AddLambda([WeakComponent]()
+			{
+				if (FFastGeoStaticMeshComponentBase* Component = WeakComponent.Get<FFastGeoStaticMeshComponentBase>())
+				{
+					if (Component->OnResourceUpdated.IsValid())
+					{
+						UMeshPaintVirtualTexture* MeshPaintTexture = Cast<UMeshPaintVirtualTexture>(Component->GetStaticMeshSceneProxyDesc().MeshPaintTexture);
+						if (ensure(MeshPaintTexture))
+						{
+							MeshPaintTexture->GetOnResourceUpdated().Remove(Component->OnResourceUpdated);
+						}
+						Component->OnResourceUpdated.Reset();
+					}
+
+					// Remove the component from the container unregistration event
+					Component->GetOwnerContainer()->GetOnUnregistered().Remove(Component->OnContainerUnregistered);
+					Component->OnContainerUnregistered.Reset();
+				}
+			});
+		}
+	}
+
 	Super::InitializeDynamicProperties();
 }
 
@@ -168,6 +214,22 @@ void FFastGeoStaticMeshComponentBase::InitializeSceneProxyDescFromComponent(UAct
 	UStaticMeshComponent* StaticMeshComponent = CastChecked<UStaticMeshComponent>(Component);
 	FStaticMeshSceneProxyDesc& SceneProxyDesc = GetStaticMeshSceneProxyDesc();
 	SceneProxyDesc.InitializeFromStaticMeshComponent(StaticMeshComponent);
+
+	SceneProxyDesc.MeshPaintTexture = nullptr;
+	SceneProxyDesc.MeshPaintTextureCoordinateIndex = 0;
+	if (UTexture* Texture = StaticMeshComponent->GetMeshPaintTexture())
+	{
+		FImage Image;
+		if (Texture->Source.GetMipImage(Image, 0))
+		{
+			UMeshPaintVirtualTexture* NewTexture = NewObject<UMeshPaintVirtualTexture>(GetOwnerContainer());
+			NewTexture->Source.Init(Image);
+			NewTexture->UpdateResource();
+
+			SceneProxyDesc.MeshPaintTexture = NewTexture;
+			SceneProxyDesc.MeshPaintTextureCoordinateIndex = StaticMeshComponent->GetMeshPaintTextureCoordinateIndex();
+		}
+	}
 }
 
 void FFastGeoStaticMeshComponentBase::InitializeFromComponent(UActorComponent* Component)
@@ -188,8 +250,6 @@ void FFastGeoStaticMeshComponentBase::ResetSceneProxyDescUnsupportedProperties()
 	FStaticMeshSceneProxyDesc& SceneProxyDesc = GetStaticMeshSceneProxyDesc();
 	SceneProxyDesc.LODData = TArrayView<struct FStaticMeshComponentLODInfo>();
 	SceneProxyDesc.LODParentPrimitive = nullptr;
-	SceneProxyDesc.MeshPaintTexture = nullptr;
-	SceneProxyDesc.MeshPaintTextureCoordinateIndex = 0;
 #if STATICMESH_ENABLE_DEBUG_RENDERING
 	SceneProxyDesc.bDrawMeshCollisionIfComplex = false;
 	SceneProxyDesc.bDrawMeshCollisionIfSimple = false;
@@ -233,7 +293,6 @@ void FFastGeoStaticMeshComponentBase::InitializeSceneProxyDescDynamicProperties(
 
 	// Add LODData support
 	// Add LODParentPrimitive support
-	// Add MeshPaintTexture/MeshPaintTextureCoordinateIndex support
 }
 
 FPrimitiveSceneProxy* FFastGeoStaticMeshComponentBase::CreateSceneProxy(ESceneProxyCreationError* OutError)
diff --git a/Engine/Source/Runtime/Engine/Classes/VT/MeshPaintVirtualTexture.h b/Engine/Source/Runtime/Engine/Classes/VT/MeshPaintVirtualTexture.h
--- a/Engine/Source/Runtime/Engine/Classes/VT/MeshPaintVirtualTexture.h
+++ b/Engine/Source/Runtime/Engine/Classes/VT/MeshPaintVirtualTexture.h
@@ -33,6 +33,14 @@ class UMeshPaintVirtualTexture : public UTexture2D
 	virtual void ClearCachedCookedPlatformData(const ITargetPlatform* TargetPlatform) override;
 #endif
 	//~ End UTexture Interface.
+
+	/** Returns the delegate called once UpdateResourceWithParams has completed. */
+	DECLARE_MULTICAST_DELEGATE(FOnResourceUpdated);
+	FOnResourceUpdated& GetOnResourceUpdated() { return OnResourceUpdated; }
+
+private:
+	/** Delegate that this texture will call once UpdateResourceWithParams has completed. */
+	FOnResourceUpdated OnResourceUpdated;
 };
 
 namespace MeshPaintVirtualTexture
diff --git a/Engine/Source/Runtime/Engine/Private/VT/MeshPaintVirtualTexture.cpp b/Engine/Source/Runtime/Engine/Private/VT/MeshPaintVirtualTexture.cpp
--- a/Engine/Source/Runtime/Engine/Private/VT/MeshPaintVirtualTexture.cpp
+++ b/Engine/Source/Runtime/Engine/Private/VT/MeshPaintVirtualTexture.cpp
@@ -321,6 +321,9 @@ void UMeshPaintVirtualTexture::UpdateResourceWithParams(EUpdateResourceFlags InF
 	{
 		PrimitiveComponent->MarkRenderStateDirty();
 	}
+
+	// Notify that we are done updating resource
+	OnResourceUpdated.Broadcast();
 }
 
 #if WITH_EDITOR
