FasGeo - Avoid keeping a reference to InstanceSceneDataBuffers as it forces the memory to only be freed from the game thread when the container's destructor is called (during garbage collection)

#Synced-CL 48020879

diff --git a/Engine/Plugins/Experimental/FastGeoStreaming/Source/FastGeoStreaming/Internal/FastGeoInstancedStaticMeshComponent.h b/Engine/Plugins/Experimental/FastGeoStreaming/Source/FastGeoStreaming/Internal/FastGeoInstancedStaticMeshComponent.h
--- a/Engine/Plugins/Experimental/FastGeoStreaming/Source/FastGeoStreaming/Internal/FastGeoInstancedStaticMeshComponent.h
+++ b/Engine/Plugins/Experimental/FastGeoStreaming/Source/FastGeoStreaming/Internal/FastGeoInstancedStaticMeshComponent.h
@@ -55,7 +55,7 @@ protected:
 #endif
 	//~ End FFastGeoStaticMeshComponentBase interface
 
-	TSharedPtr<FInstanceDataSceneProxy, ESPMode::ThreadSafe>& BuildInstanceData();
+	void BuildInstanceData();
 
 private:
 	enum class EBoundsType
@@ -78,7 +78,6 @@ private:
 	FInstancedStaticMeshSceneProxyDesc SceneProxyDesc{};
 
 	// Transient data
-	TSharedPtr<FInstanceDataSceneProxy, ESPMode::ThreadSafe> DataProxy{};
 	TArray<float> InstanceRandomIDs;
 	
 	/** Physics representation of the instance bodies. */
diff --git a/Engine/Plugins/Experimental/FastGeoStreaming/Source/FastGeoStreaming/Private/FastGeoInstancedStaticMeshComponent.cpp b/Engine/Plugins/Experimental/FastGeoStreaming/Source/FastGeoStreaming/Private/FastGeoInstancedStaticMeshComponent.cpp
--- a/Engine/Plugins/Experimental/FastGeoStreaming/Source/FastGeoStreaming/Private/FastGeoInstancedStaticMeshComponent.cpp
+++ b/Engine/Plugins/Experimental/FastGeoStreaming/Source/FastGeoStreaming/Private/FastGeoInstancedStaticMeshComponent.cpp
@@ -86,7 +86,7 @@ void FFastGeoInstancedStaticMeshComponent::ApplyWorldTransform(const FTransform&
 	NavigationBounds = CalculateBounds(EBoundsType::NavigationBounds).GetBox();
 }
 
-TSharedPtr<FInstanceDataSceneProxy, ESPMode::ThreadSafe>& FFastGeoInstancedStaticMeshComponent::BuildInstanceData()
+void FFastGeoInstancedStaticMeshComponent::BuildInstanceData()
 {
 	FInstanceSceneDataBuffers InstanceSceneDataBuffers{};
 	FInstanceSceneDataBuffers::FAccessTag AccessTag(PointerHash(this));
@@ -150,8 +150,7 @@ TSharedPtr<FInstanceDataSceneProxy, ESPMode::ThreadSafe>& FFastGeoInstancedStati
 	InstanceSceneDataBuffers.EndWriteAccess(AccessTag);
 	InstanceSceneDataBuffers.ValidateData();
 
-	DataProxy = MakeShared<FInstanceDataSceneProxy, ESPMode::ThreadSafe>(MoveTemp(InstanceSceneDataBuffers));
-	return DataProxy;
+	SceneProxyDesc.InstanceDataSceneProxy = MakeShared<FInstanceDataSceneProxy, ESPMode::ThreadSafe>(MoveTemp(InstanceSceneDataBuffers));
 }
 
 FPrimitiveSceneProxy* FFastGeoInstancedStaticMeshComponent::CreateStaticMeshSceneProxy(const Nanite::FMaterialAudit& NaniteMaterials, bool bCreateNanite)
@@ -160,7 +159,8 @@ FPrimitiveSceneProxy* FFastGeoInstancedStaticMeshComponent::CreateStaticMeshScen
 	check(SceneProxyDesc.Scene);
 	check(PerInstanceSMData.Num() > 0);
 
-	SceneProxyDesc.InstanceDataSceneProxy = BuildInstanceData();
+	BuildInstanceData();
+
 	if (bCreateNanite)
 	{
 		PrimitiveSceneData.SceneProxy = ::new Nanite::FSceneProxy(NaniteMaterials, SceneProxyDesc);
@@ -169,6 +169,11 @@ FPrimitiveSceneProxy* FFastGeoInstancedStaticMeshComponent::CreateStaticMeshScen
 	{
 		PrimitiveSceneData.SceneProxy = ::new FInstancedStaticMeshSceneProxy(SceneProxyDesc, SceneProxyDesc.FeatureLevel);
 	}
+
+	// No need to keep this around, as keeping a reference would force the various instance scene data buffers to only
+	// be released at GC time. Releasing the ref will allow the data to be freed outside of the game thread.
+	SceneProxyDesc.InstanceDataSceneProxy.Reset();
+
 	return PrimitiveSceneData.SceneProxy;
 }
 
