[FastGeo] Refactor to simplify/clarity interface of FastGeoComponent
- Removed ShouldCreateRenderState in favor of ShouldComponentAddToRenderScene
- MarkRenderStateDirty can be called even if render proxy was not previously created
- Removed BatchAddInvisibleLights now that MarkRenderStateDirty can be used to create a light that was not previously created

#Synced-CL 48520266

diff --git a/Engine/Plugins/Experimental/FastGeoStreaming/Source/FastGeoStreaming/Internal/FastGeoComponent.h b/Engine/Plugins/Experimental/FastGeoStreaming/Source/FastGeoStreaming/Internal/FastGeoComponent.h
--- a/Engine/Plugins/Experimental/FastGeoStreaming/Source/FastGeoStreaming/Internal/FastGeoComponent.h
+++ b/Engine/Plugins/Experimental/FastGeoStreaming/Source/FastGeoStreaming/Internal/FastGeoComponent.h
@@ -45,7 +45,7 @@ public:
 	virtual void InitializeFromComponent(UActorComponent* Component);
 #endif
 	int32 GetComponentIndex() const { return ComponentIndex; }
-	void MarkRenderStateDirty();
+	void MarkRenderStateDirty(bool bEvenIfNotCreated);
 	bool IsRenderStateDirty() const;
 	bool IsRenderStateDelayed() const;
 	bool IsRenderStateCreated() const;
@@ -54,8 +54,7 @@ public:
 	virtual void Serialize(FArchive& Ar);
 	virtual void InitializeDynamicProperties() {}
 	virtual void ApplyWorldTransform(const FTransform& InTransform);
-	virtual bool ShouldComponentAddToScene() const;
-	virtual bool ShouldCreateRenderState() const { return false; }
+	virtual bool ShouldComponentAddToRenderScene() const;
 	virtual void CreateRenderState(FRegisterComponentContext* Context) {}
 	virtual void DestroyRenderState(FFastGeoDestroyRenderStateContext* Context) {}
 	virtual void UpdateVisibility() {}
diff --git a/Engine/Plugins/Experimental/FastGeoStreaming/Source/FastGeoStreaming/Internal/FastGeoDecalComponent.h b/Engine/Plugins/Experimental/FastGeoStreaming/Source/FastGeoStreaming/Internal/FastGeoDecalComponent.h
--- a/Engine/Plugins/Experimental/FastGeoStreaming/Source/FastGeoStreaming/Internal/FastGeoDecalComponent.h
+++ b/Engine/Plugins/Experimental/FastGeoStreaming/Source/FastGeoStreaming/Internal/FastGeoDecalComponent.h
@@ -23,7 +23,7 @@ public:
 	//~ Begin FFastGeoComponent interface
 	virtual void Serialize(FArchive& Ar) override;
 	virtual void ApplyWorldTransform(const FTransform& InTransform) override;
-	virtual bool ShouldCreateRenderState() const override;
+	virtual bool ShouldComponentAddToRenderScene() const override;
 	virtual void CreateRenderState(FRegisterComponentContext* Context) override;
 	virtual void DestroyRenderState(FFastGeoDestroyRenderStateContext* Context) override;
 	virtual void UpdateVisibility() override;
diff --git a/Engine/Plugins/Experimental/FastGeoStreaming/Source/FastGeoStreaming/Internal/FastGeoLocalLightComponent.h b/Engine/Plugins/Experimental/FastGeoStreaming/Source/FastGeoStreaming/Internal/FastGeoLocalLightComponent.h
--- a/Engine/Plugins/Experimental/FastGeoStreaming/Source/FastGeoStreaming/Internal/FastGeoLocalLightComponent.h
+++ b/Engine/Plugins/Experimental/FastGeoStreaming/Source/FastGeoStreaming/Internal/FastGeoLocalLightComponent.h
@@ -20,7 +20,7 @@ public:
 
 	//~ Begin FFastGeoComponent interface
 	virtual void Serialize(FArchive& Ar) override;
-	virtual bool ShouldCreateRenderState() const override;
+	virtual bool ShouldComponentAddToRenderScene() const override;
 	virtual void CreateRenderState(FRegisterComponentContext* Context) override;
 	virtual void DestroyRenderState(FFastGeoDestroyRenderStateContext* Context) override;
 	virtual void UpdateVisibility() override;
@@ -32,15 +32,17 @@ public:
 protected:
 	virtual FLocalLightSceneProxyDesc& GetLocalLightSceneProxyDesc() = 0;
 	virtual const FLocalLightSceneProxyDesc& GetLocalLightSceneProxyDesc() const = 0;
-	virtual FLightSceneProxy* CreateSceneProxy() = 0;
+	virtual FLightSceneProxy* CreateTypedSceneProxy() = 0;
 
 	void InitializeSceneProxyDescDynamicProperties();
 	bool UpdateVisibilityInternal();
 
 private:
+	FLightSceneProxy* CreateSceneProxy();
+
 	// Persistent data
 	bool bAffectsWorld = true;
-	bool bComponentIsVisible = true;
+	bool bSourceComponentIsVisible = true;
 
 	// Transient data
 	bool bIsVisible = true;
diff --git a/Engine/Plugins/Experimental/FastGeoStreaming/Source/FastGeoStreaming/Internal/FastGeoPointLightComponent.h b/Engine/Plugins/Experimental/FastGeoStreaming/Source/FastGeoStreaming/Internal/FastGeoPointLightComponent.h
--- a/Engine/Plugins/Experimental/FastGeoStreaming/Source/FastGeoStreaming/Internal/FastGeoPointLightComponent.h
+++ b/Engine/Plugins/Experimental/FastGeoStreaming/Source/FastGeoStreaming/Internal/FastGeoPointLightComponent.h
@@ -19,15 +19,17 @@ public:
 	virtual ~FFastGeoPointLightComponentBase() = default;
 
 protected:
+	//~ Begin FFastGeoComponent interface
+	bool ShouldComponentAddToRenderScene() const;
+	//~ End FFastGeoComponent interface
+
 	//~ Begin FFastGeoLocalLightComponent interface
 	virtual FLocalLightSceneProxyDesc& GetLocalLightSceneProxyDesc() override { return GetPointLightSceneProxyDesc(); }
 	virtual const FLocalLightSceneProxyDesc& GetLocalLightSceneProxyDesc() const override { return GetPointLightSceneProxyDesc(); }
-	virtual FLightSceneProxy* CreateSceneProxy() override;
 	//~ End FFastGeoLocalLightComponent interface
 
 	virtual FPointLightSceneProxyDesc& GetPointLightSceneProxyDesc() = 0;
 	virtual const FPointLightSceneProxyDesc& GetPointLightSceneProxyDesc() const = 0;
-	virtual FLightSceneProxy* CreateTypedSceneProxy() = 0;
 
 	bool IsPointLightSupported() const;
 };
@@ -50,10 +52,13 @@ public:
 	//~ End FFastGeoComponent interface
 
 protected:
+	//~ Begin FFastGeoLightComponent interface
+	virtual FLightSceneProxy* CreateTypedSceneProxy() override;
+	//~ End FFastGeoLightComponent interface
+
 	//~ Begin FFastGeoPointLightComponentBase interface
 	virtual FPointLightSceneProxyDesc& GetPointLightSceneProxyDesc() override { return SceneProxyDesc; }
 	virtual const FPointLightSceneProxyDesc& GetPointLightSceneProxyDesc() const override { return SceneProxyDesc; }
-	virtual FLightSceneProxy* CreateTypedSceneProxy() override;
 	//~ End FFastGeoPointLightComponentBase interface
 
 private:
diff --git a/Engine/Plugins/Experimental/FastGeoStreaming/Source/FastGeoStreaming/Internal/FastGeoPrimitiveComponent.h b/Engine/Plugins/Experimental/FastGeoStreaming/Source/FastGeoStreaming/Internal/FastGeoPrimitiveComponent.h
--- a/Engine/Plugins/Experimental/FastGeoStreaming/Source/FastGeoStreaming/Internal/FastGeoPrimitiveComponent.h
+++ b/Engine/Plugins/Experimental/FastGeoStreaming/Source/FastGeoStreaming/Internal/FastGeoPrimitiveComponent.h
@@ -141,6 +141,7 @@ public:
 	virtual void OnAsyncDestroyPhysicsStateEnd_GameThread() override;
 	virtual void OnAsyncDestroyPhysicsState() override;
 	virtual bool IsCollisionEnabled() const override;
+	virtual bool ShouldComponentAddToRenderScene() const override;
 	virtual void DestroyRenderState(FFastGeoDestroyRenderStateContext* Context) override;
 	virtual void MarkPrecachePSOsRequired() override;
 #if WITH_EDITOR
diff --git a/Engine/Plugins/Experimental/FastGeoStreaming/Source/FastGeoStreaming/Internal/FastGeoRectLightComponent.h b/Engine/Plugins/Experimental/FastGeoStreaming/Source/FastGeoStreaming/Internal/FastGeoRectLightComponent.h
--- a/Engine/Plugins/Experimental/FastGeoStreaming/Source/FastGeoStreaming/Internal/FastGeoRectLightComponent.h
+++ b/Engine/Plugins/Experimental/FastGeoStreaming/Source/FastGeoStreaming/Internal/FastGeoRectLightComponent.h
@@ -23,10 +23,13 @@ public:
 	//~ End FFastGeoComponent interface
 
 protected:
+	//~ Begin FFastGeoLightComponent interface
+	virtual FLightSceneProxy* CreateTypedSceneProxy() override;
+	//~ End FFastGeoLightComponent interface
+
 	//~ Begin FFastGeoLocalLightComponent interface
 	virtual FLocalLightSceneProxyDesc& GetLocalLightSceneProxyDesc() override { return SceneProxyDesc; }
 	virtual const FLocalLightSceneProxyDesc& GetLocalLightSceneProxyDesc() const override { return SceneProxyDesc; }
-	virtual FLightSceneProxy* CreateSceneProxy() override;
 	//~ End FFastGeoLocalLightComponent interface
 
 private:
diff --git a/Engine/Plugins/Experimental/FastGeoStreaming/Source/FastGeoStreaming/Internal/FastGeoSpotLightComponent.h b/Engine/Plugins/Experimental/FastGeoStreaming/Source/FastGeoStreaming/Internal/FastGeoSpotLightComponent.h
--- a/Engine/Plugins/Experimental/FastGeoStreaming/Source/FastGeoStreaming/Internal/FastGeoSpotLightComponent.h
+++ b/Engine/Plugins/Experimental/FastGeoStreaming/Source/FastGeoStreaming/Internal/FastGeoSpotLightComponent.h
@@ -23,10 +23,13 @@ public:
 	//~ End FFastGeoComponent interface
 
 protected:
+	//~ Begin FFastGeoLightComponent interface
+	virtual FLightSceneProxy* CreateTypedSceneProxy() override;
+	//~ End FFastGeoLightComponent interface
+
 	//~ Begin FFastGeoPointLightComponentBase interface
 	virtual FPointLightSceneProxyDesc& GetPointLightSceneProxyDesc() override { return SceneProxyDesc; }
 	virtual const FPointLightSceneProxyDesc& GetPointLightSceneProxyDesc() const override { return SceneProxyDesc; }
-	virtual FLightSceneProxy* CreateTypedSceneProxy() override;
 	//~ End FFastGeoPointLightComponentBase interface
 
 private:
diff --git a/Engine/Plugins/Experimental/FastGeoStreaming/Source/FastGeoStreaming/Private/FastGeoComponent.cpp b/Engine/Plugins/Experimental/FastGeoStreaming/Source/FastGeoStreaming/Private/FastGeoComponent.cpp
--- a/Engine/Plugins/Experimental/FastGeoStreaming/Source/FastGeoStreaming/Private/FastGeoComponent.cpp
+++ b/Engine/Plugins/Experimental/FastGeoStreaming/Source/FastGeoStreaming/Private/FastGeoComponent.cpp
@@ -153,9 +153,9 @@ bool FFastGeoComponent::IsRenderStateDirty() const
 	return bRenderStateDirty;
 }
 
-void FFastGeoComponent::MarkRenderStateDirty()
+void FFastGeoComponent::MarkRenderStateDirty(bool bEvenIfNotCreated)
 {
-	if ((IsRenderStateCreated() || IsRenderStateDelayed()) && !IsRenderStateDirty())
+	if ((IsRenderStateCreated() || IsRenderStateDelayed() || bEvenIfNotCreated) && !IsRenderStateDirty())
 	{
 		bRenderStateDirty = true;
 
@@ -164,8 +164,13 @@ void FFastGeoComponent::MarkRenderStateDirty()
 	}
 }
 
-bool FFastGeoComponent::ShouldComponentAddToScene() const
+bool FFastGeoComponent::ShouldComponentAddToRenderScene() const
 {
+	if (!FApp::CanEverRender())
+	{
+		return false;
+	}
+
 	// If the detail mode setting allows it, add it to the scene.
 	return DetailMode <= GetCachedScalabilityCVars().DetailMode;
 }
\ No newline at end of file
diff --git a/Engine/Plugins/Experimental/FastGeoStreaming/Source/FastGeoStreaming/Private/FastGeoContainer.cpp b/Engine/Plugins/Experimental/FastGeoStreaming/Source/FastGeoStreaming/Private/FastGeoContainer.cpp
--- a/Engine/Plugins/Experimental/FastGeoStreaming/Source/FastGeoStreaming/Private/FastGeoContainer.cpp
+++ b/Engine/Plugins/Experimental/FastGeoStreaming/Source/FastGeoStreaming/Private/FastGeoContainer.cpp
@@ -206,7 +206,7 @@ void UFastGeoContainer::Register()
 					Component.ApplyWorldTransform(Transform);
 				}
 
-				if (Component.ShouldCreateRenderState())
+				if (Component.ShouldComponentAddToRenderScene())
 				{
 					PendingCreate.RenderState.ComponentsToProcess.Add(&Component);
 				}
diff --git a/Engine/Plugins/Experimental/FastGeoStreaming/Source/FastGeoStreaming/Private/FastGeoDecalComponent.cpp b/Engine/Plugins/Experimental/FastGeoStreaming/Source/FastGeoStreaming/Private/FastGeoDecalComponent.cpp
--- a/Engine/Plugins/Experimental/FastGeoStreaming/Source/FastGeoStreaming/Private/FastGeoDecalComponent.cpp
+++ b/Engine/Plugins/Experimental/FastGeoStreaming/Source/FastGeoStreaming/Private/FastGeoDecalComponent.cpp
@@ -84,7 +84,7 @@ void FFastGeoDecalComponent::ApplyWorldTransform(const FTransform& InTransform)
 void FFastGeoDecalComponent::UpdateVisibility()
 {
 	UpdateVisibilityInternal();
-	MarkRenderStateDirty();
+	MarkRenderStateDirty(true);
 }
 
 void FFastGeoDecalComponent::UpdateVisibilityInternal()
@@ -95,13 +95,19 @@ void FFastGeoDecalComponent::UpdateVisibilityInternal()
 #endif
 }
 
-bool FFastGeoDecalComponent::ShouldCreateRenderState() const
+bool FFastGeoDecalComponent::ShouldComponentAddToRenderScene() const
 {
-	return FApp::CanEverRender() && ShouldComponentAddToScene() && SceneProxyDesc.bDrawInGame;
+	return Super::ShouldComponentAddToRenderScene() && SceneProxyDesc.bDrawInGame;
 }
 
 void FFastGeoDecalComponent::CreateRenderState(FRegisterComponentContext* Context)
 {
+	if (!ShouldComponentAddToRenderScene())
+	{
+		ProxyState = EProxyCreationState::None;
+		return;
+	}
+
 	FWriteScopeLock WriteLock(*Lock.Get());
 	check(ProxyState == EProxyCreationState::None);
 	check(!SceneProxy);
@@ -279,6 +285,6 @@ void FFastGeoDecalComponent::OnPrecacheFinished(int32 JobSetThatJustCompleted)
 	QUICK_SCOPE_CYCLE_COUNTER(STAT_PSOPrecacheFinishedTask);
 	int32 CurrJobSetCompleted = LatestPSOPrecacheJobSetCompleted.load();
 	while (CurrJobSetCompleted < JobSetThatJustCompleted && !LatestPSOPrecacheJobSetCompleted.compare_exchange_weak(CurrJobSetCompleted, JobSetThatJustCompleted)) {}
-	MarkRenderStateDirty();
+	MarkRenderStateDirty(false);
 }
 #endif
\ No newline at end of file
diff --git a/Engine/Plugins/Experimental/FastGeoStreaming/Source/FastGeoStreaming/Private/FastGeoLocalLightComponent.cpp b/Engine/Plugins/Experimental/FastGeoStreaming/Source/FastGeoStreaming/Private/FastGeoLocalLightComponent.cpp
--- a/Engine/Plugins/Experimental/FastGeoStreaming/Source/FastGeoStreaming/Private/FastGeoLocalLightComponent.cpp
+++ b/Engine/Plugins/Experimental/FastGeoStreaming/Source/FastGeoStreaming/Private/FastGeoLocalLightComponent.cpp
@@ -15,7 +15,7 @@ const FFastGeoElementType FFastGeoLocalLightComponent::Type(&FFastGeoComponent::
 FFastGeoLocalLightComponent::FFastGeoLocalLightComponent(int32 InComponentIndex, FFastGeoElementType InType)
 	: Super(InComponentIndex, InType)
 	, bAffectsWorld(true)
-	, bComponentIsVisible(true)
+	, bSourceComponentIsVisible(true)
 	, bIsVisible(true)
 	, SceneProxy(nullptr)
 {
@@ -28,7 +28,7 @@ void FFastGeoLocalLightComponent::InitializeFromComponent(UActorComponent* Compo
 
 	ULocalLightComponent* LocalLightComponent = CastChecked<ULocalLightComponent>(Component);
 	bAffectsWorld = LocalLightComponent->bAffectsWorld;
-	bComponentIsVisible = LocalLightComponent->IsVisible() && LocalLightComponent->GetOwner() && !LocalLightComponent->GetOwner()->IsHidden();
+	bSourceComponentIsVisible = LocalLightComponent->IsVisible() && LocalLightComponent->GetOwner() && !LocalLightComponent->GetOwner()->IsHidden();
 }
 #endif
 
@@ -38,7 +38,7 @@ void FFastGeoLocalLightComponent::Serialize(FArchive& Ar)
 
 	// Serialize persistent data
 	Ar << bAffectsWorld;
-	Ar << bComponentIsVisible;
+	Ar << bSourceComponentIsVisible;
 
 	// Serialize persistent data from FLocalLightSceneProxyDesc
 	FLocalLightSceneProxyDesc& SceneProxyDesc = GetLocalLightSceneProxyDesc();
@@ -51,47 +51,35 @@ void FFastGeoLocalLightComponent::Serialize(FArchive& Ar)
 	}
 }
 
-bool FFastGeoLocalLightComponent::ShouldCreateRenderState() const
+bool FFastGeoLocalLightComponent::ShouldComponentAddToRenderScene() const
 {
-	if (!FApp::CanEverRender())
-	{
-		return false;
-	}
-
-	// If the detail mode setting allows it, add it to the scene.
-	if (!ShouldComponentAddToScene())
-	{
-		return false;
-	}
-
-	return bAffectsWorld && bComponentIsVisible;
+	return Super::ShouldComponentAddToRenderScene() && bAffectsWorld && bSourceComponentIsVisible && bIsVisible;
 }
 
 void FFastGeoLocalLightComponent::CreateRenderState(FRegisterComponentContext* Context)
 {
+	InitializeSceneProxyDescDynamicProperties();
+	if (!ShouldComponentAddToRenderScene())
+	{
+		ProxyState = EProxyCreationState::None;
+		return;
+	}
+
 	FWriteScopeLock WriteLock(*Lock.Get());
 	check(ProxyState == EProxyCreationState::None);
 	ProxyState = EProxyCreationState::Creating;
 	bRenderStateDirty = false;
 
-	InitializeSceneProxyDescDynamicProperties();
 	check(!SceneProxy);
 	SceneProxy = CreateSceneProxy();
 	if (SceneProxy)
 	{
+		check(bIsVisible);
 		FLightSceneDesc LightSceneDesc{ .Transform = WorldTransform, .SceneProxy = SceneProxy};
 		FSceneInterface* Scene = GetScene();
 		check(Scene);
-		if (bIsVisible)
-		{
-			Scene->BatchAddLights(MakeArrayView(&LightSceneDesc, 1));
-		}
-		else
-		{
-			// This should be revisited to avoid needing to create an invisible light.
-			// Note that this is currently only possible when FastGeo.Show 0 is set.
-			Scene->BatchAddInvisibleLights(MakeArrayView(&LightSceneDesc, 1));
-		}
+
+		Scene->BatchAddLights(MakeArrayView(&LightSceneDesc, 1));
 
 		check(LightSceneDesc.SceneProxy);
 		ProxyState = EProxyCreationState::Created;
@@ -102,6 +90,17 @@ void FFastGeoLocalLightComponent::CreateRenderState(FRegisterComponentContext* C
 	}
 }
 
+FLightSceneProxy* FFastGeoLocalLightComponent::CreateSceneProxy()
+{
+	check(ShouldComponentAddToRenderScene());
+
+	FSceneInterface* Scene = GetScene();
+	check(Scene);
+	FLocalLightSceneProxyDesc& SceneProxyDesc = GetLocalLightSceneProxyDesc();
+	SceneProxyDesc.SceneInterface = Scene;
+	return CreateTypedSceneProxy();
+}
+
 void FFastGeoLocalLightComponent::DestroyRenderState(FFastGeoDestroyRenderStateContext* Context)
 {
 	FWriteScopeLock WriteLock(*Lock.Get());
@@ -128,7 +127,7 @@ void FFastGeoLocalLightComponent::UpdateVisibility()
 {
 	if (UpdateVisibilityInternal())
 	{
-		MarkRenderStateDirty();
+		MarkRenderStateDirty(true);
 	}
 }
 
diff --git a/Engine/Plugins/Experimental/FastGeoStreaming/Source/FastGeoStreaming/Private/FastGeoPointLightComponent.cpp b/Engine/Plugins/Experimental/FastGeoStreaming/Source/FastGeoStreaming/Private/FastGeoPointLightComponent.cpp
--- a/Engine/Plugins/Experimental/FastGeoStreaming/Source/FastGeoStreaming/Private/FastGeoPointLightComponent.cpp
+++ b/Engine/Plugins/Experimental/FastGeoStreaming/Source/FastGeoStreaming/Private/FastGeoPointLightComponent.cpp
@@ -14,13 +14,9 @@ FFastGeoPointLightComponentBase::FFastGeoPointLightComponentBase(int32 InCompone
 {
 }
 
-FLightSceneProxy* FFastGeoPointLightComponentBase::CreateSceneProxy()
+bool FFastGeoPointLightComponentBase::ShouldComponentAddToRenderScene() const
 {
-	if (ShouldComponentAddToScene() && IsPointLightSupported())
-	{
-		return CreateTypedSceneProxy();
-	}
-	return nullptr;
+	return Super::ShouldComponentAddToRenderScene() && IsPointLightSupported();
 }
 
 bool FFastGeoPointLightComponentBase::IsPointLightSupported() const
@@ -57,8 +53,6 @@ void FFastGeoPointLightComponent::InitializeFromComponent(UActorComponent* Compo
 
 FLightSceneProxy* FFastGeoPointLightComponent::CreateTypedSceneProxy()
 {
-	FSceneInterface* Scene = GetScene();
-	check(Scene);
-	SceneProxyDesc.SceneInterface = Scene;
+	check(SceneProxyDesc.SceneInterface == GetScene());
 	return new FPointLightSceneProxy(SceneProxyDesc);
 }
\ No newline at end of file
diff --git a/Engine/Plugins/Experimental/FastGeoStreaming/Source/FastGeoStreaming/Private/FastGeoPrimitiveComponent.cpp b/Engine/Plugins/Experimental/FastGeoStreaming/Source/FastGeoStreaming/Private/FastGeoPrimitiveComponent.cpp
--- a/Engine/Plugins/Experimental/FastGeoStreaming/Source/FastGeoStreaming/Private/FastGeoPrimitiveComponent.cpp
+++ b/Engine/Plugins/Experimental/FastGeoStreaming/Source/FastGeoStreaming/Private/FastGeoPrimitiveComponent.cpp
@@ -296,6 +296,12 @@ FPrimitiveSceneDesc FFastGeoPrimitiveComponent::BuildSceneDesc()
 
 void FFastGeoPrimitiveComponent::CreateRenderState(FRegisterComponentContext* Context)
 {
+	if (!ShouldComponentAddToRenderScene())
+	{
+		ProxyState = EProxyCreationState::None;
+		return;
+	}
+
 	FWriteScopeLock WriteLock(*Lock.Get());
 	check(ProxyState == EProxyCreationState::None);
 	ProxyState = EProxyCreationState::Creating;
@@ -476,15 +482,9 @@ bool FFastGeoPrimitiveComponent::IsRenderStateCreated() const
 	return Super::IsRenderStateCreated();
 }
 
-bool FFastGeoPrimitiveComponent::ShouldCreateRenderState() const
+bool FFastGeoPrimitiveComponent::ShouldComponentAddToRenderScene() const
 {
-	if (!FApp::CanEverRender())
-	{
-		return false;
-	}
-
-	// If the detail mode setting allows it, add it to the scene.
-	if (!ShouldComponentAddToScene())
+	if (!Super::ShouldComponentAddToRenderScene())
 	{
 		return false;
 	}
@@ -494,6 +494,11 @@ bool FFastGeoPrimitiveComponent::ShouldCreateRenderState() const
 	return bShouldCreateRenderState;
 }
 
+bool FFastGeoPrimitiveComponent::ShouldCreateRenderState() const
+{
+	return ShouldComponentAddToRenderScene();
+}
+
 bool FFastGeoPrimitiveComponent::IsRegistered() const
 {
 	return Super::IsRegistered();
@@ -542,7 +547,7 @@ FPrimitiveSceneProxy* FFastGeoPrimitiveComponent::GetSceneProxy() const
 
 void FFastGeoPrimitiveComponent::MarkRenderStateDirty()
 {
-	Super::MarkRenderStateDirty();
+	Super::MarkRenderStateDirty(false);
 }
 
 void FFastGeoPrimitiveComponent::DestroyRenderState()
diff --git a/Engine/Plugins/Experimental/FastGeoStreaming/Source/FastGeoStreaming/Private/FastGeoRectLightComponent.cpp b/Engine/Plugins/Experimental/FastGeoStreaming/Source/FastGeoStreaming/Private/FastGeoRectLightComponent.cpp
--- a/Engine/Plugins/Experimental/FastGeoStreaming/Source/FastGeoStreaming/Private/FastGeoRectLightComponent.cpp
+++ b/Engine/Plugins/Experimental/FastGeoStreaming/Source/FastGeoStreaming/Private/FastGeoRectLightComponent.cpp
@@ -24,14 +24,8 @@ void FFastGeoRectLightComponent::InitializeFromComponent(UActorComponent* Compon
 }
 #endif
 
-FLightSceneProxy* FFastGeoRectLightComponent::CreateSceneProxy()
+FLightSceneProxy* FFastGeoRectLightComponent::CreateTypedSceneProxy()
 {
-	if (ShouldComponentAddToScene())
-	{
-		FSceneInterface* Scene = GetScene();
-		check(Scene);
-		SceneProxyDesc.SceneInterface = Scene;
-		return new FRectLightSceneProxy(SceneProxyDesc);
-	}
-	return nullptr;
+	check(SceneProxyDesc.SceneInterface == GetScene());
+	return new FRectLightSceneProxy(SceneProxyDesc);
 }
\ No newline at end of file
diff --git a/Engine/Plugins/Experimental/FastGeoStreaming/Source/FastGeoStreaming/Private/FastGeoSpotLightComponent.cpp b/Engine/Plugins/Experimental/FastGeoStreaming/Source/FastGeoStreaming/Private/FastGeoSpotLightComponent.cpp
--- a/Engine/Plugins/Experimental/FastGeoStreaming/Source/FastGeoStreaming/Private/FastGeoSpotLightComponent.cpp
+++ b/Engine/Plugins/Experimental/FastGeoStreaming/Source/FastGeoStreaming/Private/FastGeoSpotLightComponent.cpp
@@ -25,8 +25,6 @@ void FFastGeoSpotLightComponent::InitializeFromComponent(UActorComponent* Compon
 
 FLightSceneProxy* FFastGeoSpotLightComponent::CreateTypedSceneProxy()
 {
-	FSceneInterface* Scene = GetScene();
-	check(Scene);
-	SceneProxyDesc.SceneInterface = Scene;
+	check(SceneProxyDesc.SceneInterface == GetScene());
 	return new FSpotLightSceneProxy(SceneProxyDesc);
 }
\ No newline at end of file
