[FastGeo] Added PSO Precaching for FastGeoLights
- PSO Precaching refactor following refactor done in CL 48140040 that allows more code reuse.
- Refactored FastGeoDecal & FastGeoPrimitive to use FPSOPrecacheComponentData.
- PSO Precaching now handled by FastGeo job queue to avoid precaching running concurrently with render proxy creation

#Synced-CL 48520266

diff --git a/Engine/Plugins/Experimental/FastGeoStreaming/Source/FastGeoStreaming/Internal/FastGeoComponent.h b/Engine/Plugins/Experimental/FastGeoStreaming/Source/FastGeoStreaming/Internal/FastGeoComponent.h
--- a/Engine/Plugins/Experimental/FastGeoStreaming/Source/FastGeoStreaming/Internal/FastGeoComponent.h
+++ b/Engine/Plugins/Experimental/FastGeoStreaming/Source/FastGeoStreaming/Internal/FastGeoComponent.h
@@ -4,6 +4,7 @@
 
 #include "CoreMinimal.h"
 #include "PSOPrecacheFwd.h"
+#include "PSOPrecacheComponent.h"
 #include "IFastGeoElement.h"
 #include "FastGeoElementType.h"
 #include "Physics/Experimental/AsyncPhysicsStateProcessorInterface.h"
@@ -58,17 +59,19 @@ public:
 	virtual void CreateRenderState(FRegisterComponentContext* Context) {}
 	virtual void DestroyRenderState(FFastGeoDestroyRenderStateContext* Context) {}
 	virtual void UpdateVisibility() {}
-	virtual void MarkPrecachePSOsRequired() {}
-	virtual void PrecachePSOs() {}
-#if !WITH_EDITOR && UE_WITH_PSO_PRECACHING
-	virtual void OnPrecacheFinished(int32 JobSetThatJustCompleted) {}
-#endif
 	virtual void OnAsyncCreatePhysicsState();
 	virtual void OnAsyncDestroyPhysicsStateBegin_GameThread();
 	virtual void OnAsyncDestroyPhysicsState();
 	virtual void OnAsyncDestroyPhysicsStateEnd_GameThread();
 	void OnAsyncCreatePhysicsStateEnd_GameThread();
 
+#if !WITH_EDITOR && UE_WITH_PSO_PRECACHING
+	virtual bool ShouldPrecachePSOs() const;
+	virtual void PrecachePSOs_Concurrent() {}
+	virtual FPSOPrecacheComponentData& GetPSOPrecacheComponentData() = 0;
+	virtual const FPSOPrecacheComponentData& GetPSOPrecacheComponentData() const = 0;
+#endif
+
 	FFastGeoComponentCluster* GetOwnerComponentCluster() const;
 	UFastGeoContainer* GetOwnerContainer() const;
 	UWorld* GetWorld() const;
diff --git a/Engine/Plugins/Experimental/FastGeoStreaming/Source/FastGeoStreaming/Internal/FastGeoContainer.h b/Engine/Plugins/Experimental/FastGeoStreaming/Source/FastGeoStreaming/Internal/FastGeoContainer.h
--- a/Engine/Plugins/Experimental/FastGeoStreaming/Source/FastGeoStreaming/Internal/FastGeoContainer.h
+++ b/Engine/Plugins/Experimental/FastGeoStreaming/Source/FastGeoStreaming/Internal/FastGeoContainer.h
@@ -148,10 +148,13 @@ public:
 
 	void InitializeDynamicProperties(bool bInitForPlay = false);
 
+	void OnPrecachePSOsBegin_GameThread() {}
+	void OnPrecachePSOsEnd_GameThread() {}
 	void OnCreateRenderStateBegin_GameThread();
 	void OnDestroyRenderStateBegin_GameThread();
 	void OnCreateRenderStateEnd_GameThread();
 	void OnDestroyRenderStateEnd_GameThread();
+	void OnPrecachePSOs_Concurrent();
 	void OnCreateRenderState_Concurrent();
 	void OnDestroyRenderState_Concurrent();
 
@@ -296,7 +299,6 @@ private:
 
 #if !WITH_EDITOR && UE_WITH_PSO_PRECACHING
 	bool bPrecachedPSOs = false;
-	UE::Tasks::TTask<void> PrecachePSOsTask;
 #endif
 
 	// Transient data
diff --git a/Engine/Plugins/Experimental/FastGeoStreaming/Source/FastGeoStreaming/Internal/FastGeoDecalComponent.h b/Engine/Plugins/Experimental/FastGeoStreaming/Source/FastGeoStreaming/Internal/FastGeoDecalComponent.h
--- a/Engine/Plugins/Experimental/FastGeoStreaming/Source/FastGeoStreaming/Internal/FastGeoDecalComponent.h
+++ b/Engine/Plugins/Experimental/FastGeoStreaming/Source/FastGeoStreaming/Internal/FastGeoDecalComponent.h
@@ -27,12 +27,14 @@ public:
 	virtual void CreateRenderState(FRegisterComponentContext* Context) override;
 	virtual void DestroyRenderState(FFastGeoDestroyRenderStateContext* Context) override;
 	virtual void UpdateVisibility() override;
-	virtual void PrecachePSOs() override;
-#if !WITH_EDITOR && UE_WITH_PSO_PRECACHING
-	virtual void OnPrecacheFinished(int32 JobSetThatJustCompleted) override;
-#endif
 #if WITH_EDITOR
 	virtual void InitializeFromComponent(UActorComponent* Component) override;
+#endif
+#if !WITH_EDITOR && UE_WITH_PSO_PRECACHING
+	virtual bool ShouldPrecachePSOs() const override;
+	virtual void PrecachePSOs_Concurrent() override;
+	virtual FPSOPrecacheComponentData& GetPSOPrecacheComponentData() override { return PSOPrecacheComponentData; }
+	virtual const FPSOPrecacheComponentData& GetPSOPrecacheComponentData() const override { return PSOPrecacheComponentData; }
 #endif
 	//~ End FFastGeoComponent interface
 
@@ -53,8 +55,7 @@ private:
 	FTimerHandle TimerHandle_DestroyDecalComponent;
 
 #if !WITH_EDITOR && UE_WITH_PSO_PRECACHING
-	// Atomic int used to track the last PSO precache events
-	std::atomic<int> LatestPSOPrecacheJobSetCompleted = 0;
-	int32 LatestPSOPrecacheJobSet = 0;
+	/** All PSO related data for the component */
+	FPSOPrecacheComponentData PSOPrecacheComponentData;
 #endif
 };
\ No newline at end of file
diff --git a/Engine/Plugins/Experimental/FastGeoStreaming/Source/FastGeoStreaming/Internal/FastGeoLocalLightComponent.h b/Engine/Plugins/Experimental/FastGeoStreaming/Source/FastGeoStreaming/Internal/FastGeoLocalLightComponent.h
--- a/Engine/Plugins/Experimental/FastGeoStreaming/Source/FastGeoStreaming/Internal/FastGeoLocalLightComponent.h
+++ b/Engine/Plugins/Experimental/FastGeoStreaming/Source/FastGeoStreaming/Internal/FastGeoLocalLightComponent.h
@@ -16,6 +16,7 @@ public:
 	static const FFastGeoElementType Type;
 
 	FFastGeoLocalLightComponent(int32 InComponentIndex = INDEX_NONE, FFastGeoElementType InType = Type);
+	FFastGeoLocalLightComponent(const FFastGeoLocalLightComponent& Other);
 	virtual ~FFastGeoLocalLightComponent() = default;
 
 	//~ Begin FFastGeoComponent interface
@@ -29,6 +30,13 @@ public:
 #endif
 	//~ End FFastGeoComponent interface
 
+#if !WITH_EDITOR && UE_WITH_PSO_PRECACHING
+	virtual bool ShouldPrecachePSOs() const override;
+	virtual void PrecachePSOs_Concurrent() override;
+	virtual FPSOPrecacheComponentData& GetPSOPrecacheComponentData() override { return PSOPrecacheComponentData; }
+	virtual const FPSOPrecacheComponentData& GetPSOPrecacheComponentData() const override { return PSOPrecacheComponentData; }
+#endif
+
 protected:
 	virtual FLocalLightSceneProxyDesc& GetLocalLightSceneProxyDesc() = 0;
 	virtual const FLocalLightSceneProxyDesc& GetLocalLightSceneProxyDesc() const = 0;
@@ -38,7 +46,7 @@ protected:
 	bool UpdateVisibilityInternal();
 
 private:
-	FLightSceneProxy* CreateSceneProxy();
+	FLightSceneProxy* CreateSceneProxy(ESceneProxyCreationError& OutError);
 
 	// Persistent data
 	bool bAffectsWorld = true;
@@ -47,4 +55,9 @@ private:
 	// Transient data
 	bool bIsVisible = true;
 	FLightSceneProxy* SceneProxy = nullptr;
+
+#if !WITH_EDITOR && UE_WITH_PSO_PRECACHING
+	/** All PSO related data for the component */
+	FPSOPrecacheComponentData PSOPrecacheComponentData;
+#endif
 };
\ No newline at end of file
diff --git a/Engine/Plugins/Experimental/FastGeoStreaming/Source/FastGeoStreaming/Internal/FastGeoPSOPrecacheFinishedTask.h b/Engine/Plugins/Experimental/FastGeoStreaming/Source/FastGeoStreaming/Internal/FastGeoPSOPrecacheFinishedTask.h
deleted file mode 100644
diff --git a/Engine/Plugins/Experimental/FastGeoStreaming/Source/FastGeoStreaming/Internal/FastGeoPSOWeakPtrTraits.h b/Engine/Plugins/Experimental/FastGeoStreaming/Source/FastGeoStreaming/Internal/FastGeoPSOWeakPtrTraits.h
new file mode 100644
--- /dev/null
+++ b/Engine/Plugins/Experimental/FastGeoStreaming/Source/FastGeoStreaming/Internal/FastGeoPSOWeakPtrTraits.h
@@ -0,0 +1,16 @@
+// Copyright Epic Games, Inc. All Rights Reserved.
+
+#pragma once
+
+#include "PSOPrecacheComponent.h"
+#include "FastGeoWeakElement.h"
+
+// Specialization: if T : FFastGeoComponent -> use FWeakFastGeoComponent
+template<typename T>
+struct TPSOWeakPtrTraits<T, typename TEnableIf<TIsDerivedFrom<T, FFastGeoComponent>::Value>::Type>
+{
+	using WeakType = FWeakFastGeoComponent;
+
+	static WeakType Make(T* Ptr) { return WeakType(Ptr); }
+	static T* Get(const WeakType& W) { return W.Get<T>(); }
+};
\ No newline at end of file
diff --git a/Engine/Plugins/Experimental/FastGeoStreaming/Source/FastGeoStreaming/Internal/FastGeoPointLightComponent.h b/Engine/Plugins/Experimental/FastGeoStreaming/Source/FastGeoStreaming/Internal/FastGeoPointLightComponent.h
--- a/Engine/Plugins/Experimental/FastGeoStreaming/Source/FastGeoStreaming/Internal/FastGeoPointLightComponent.h
+++ b/Engine/Plugins/Experimental/FastGeoStreaming/Source/FastGeoStreaming/Internal/FastGeoPointLightComponent.h
@@ -18,6 +18,12 @@ public:
 	FFastGeoPointLightComponentBase(int32 InComponentIndex = INDEX_NONE, FFastGeoElementType InType = Type);
 	virtual ~FFastGeoPointLightComponentBase() = default;
 
+	//~ Begin FFastGeoComponent interface
+#if !WITH_EDITOR && UE_WITH_PSO_PRECACHING
+	virtual bool ShouldPrecachePSOs() const override;
+#endif
+	//~ End FFastGeoComponent interface
+
 protected:
 	//~ Begin FFastGeoComponent interface
 	bool ShouldComponentAddToRenderScene() const;
diff --git a/Engine/Plugins/Experimental/FastGeoStreaming/Source/FastGeoStreaming/Internal/FastGeoPrimitiveComponent.h b/Engine/Plugins/Experimental/FastGeoStreaming/Source/FastGeoStreaming/Internal/FastGeoPrimitiveComponent.h
--- a/Engine/Plugins/Experimental/FastGeoStreaming/Source/FastGeoStreaming/Internal/FastGeoPrimitiveComponent.h
+++ b/Engine/Plugins/Experimental/FastGeoStreaming/Source/FastGeoStreaming/Internal/FastGeoPrimitiveComponent.h
@@ -85,8 +85,6 @@ public:
 	virtual void NotifyRenderStateChanged();
 #endif
 	bool IsDrawnInGame() const;
-	bool CheckPSOPrecachingAndBoostPriority(EPSOPrecachePriority NewPSOPrecachePriority = EPSOPrecachePriority::High);
-	bool IsPSOPrecaching() const;
 	void SetCollisionEnabled(bool bInCollisionEnabled);
 	void SetCustomPrimitiveData(TConstArrayView<float> InCustomPrimitiveData);
 	void UpdateVisibility();
@@ -143,10 +141,15 @@ public:
 	virtual bool IsCollisionEnabled() const override;
 	virtual bool ShouldComponentAddToRenderScene() const override;
 	virtual void DestroyRenderState(FFastGeoDestroyRenderStateContext* Context) override;
-	virtual void MarkPrecachePSOsRequired() override;
 #if WITH_EDITOR
 	virtual void InitializeFromComponent(UActorComponent* Component) override;
 #endif
+#if !WITH_EDITOR && UE_WITH_PSO_PRECACHING
+	virtual void PrecachePSOs_Concurrent() override;
+	virtual FPSOPrecacheComponentData& GetPSOPrecacheComponentData() override { return PSOPrecacheComponentData; }
+	virtual const FPSOPrecacheComponentData& GetPSOPrecacheComponentData() const override { return PSOPrecacheComponentData; }
+#endif
+	bool CheckPSOPrecachingAndBoostPriority(EPSOPrecachePriority NewPSOPrecachePriority = EPSOPrecachePriority::High);
 	//~End FFastGeoComponent interface
 
 	//~ Begin Materials
@@ -206,31 +209,18 @@ protected:
 	/** Payload used to release BodyInstance resources in asynchronous mode (see OnAsyncDestroyPhysicsState). */
 	TOptional<FBodyInstance::FAsyncTermBodyPayload> AsyncTermBodyPayload;
 
-private:
-	bool ShouldRenderProxyFallbackToDefaultMaterial() const;
-
-#if !WITH_EDITOR && UE_WITH_PSO_PRECACHING
-public:
-	virtual void OnPrecacheFinished(int32 JobSetThatJustCompleted) override;
-
 protected:
-	void RequestRecreateRenderStateWhenPSOPrecacheFinished(const FGraphEventArray& PSOPrecacheCompileEvents);
-	void SetupPrecachePSOParams(FPSOPrecacheParams& Params);
+#if !WITH_EDITOR && UE_WITH_PSO_PRECACHING
 	virtual void CollectPSOPrecacheData(const FPSOPrecacheParams& BasePrecachePSOParams, FMaterialInterfacePSOPrecacheParamsList& OutParams) {}
+	void SetupPrecachePSOParams(FPSOPrecacheParams& Params);
 
-	// Cached array of material PSO requests which can be used to boost the priority
-	TArray<FMaterialPSOPrecacheRequestID> MaterialPSOPrecacheRequestIDs;
-	// Atomic int used to track the last PSO precache events
-	std::atomic<int> LatestPSOPrecacheJobSetCompleted = 0;
-	int32 LatestPSOPrecacheJobSet = 0;
-	// Helper flag to check if PSOs have been precached already
-	std::atomic<bool> bPSOPrecacheCalled = false;
-	std::atomic<bool> bPSOPrecacheRequired = false;
-	// PSOs requested priority
-	std::atomic<EPSOPrecachePriority> PSOPrecacheRequestPriority = EPSOPrecachePriority::Medium;
-	static_assert((int)EPSOPrecachePriority::Highest < 1 << 2);
+	/** All PSO related data for the component */
+	FPSOPrecacheComponentData PSOPrecacheComponentData;
 #endif
 
+private:
+	bool ShouldRenderProxyFallbackToDefaultMaterial() const;
+
 	friend class FFastGeoComponentCluster;
 	friend class FPrimitiveComponentHelper;
 	friend class FFastGeoPhysicsBodyInstanceOwner;
diff --git a/Engine/Plugins/Experimental/FastGeoStreaming/Source/FastGeoStreaming/Private/FastGeoAsyncRenderStateJobQueue.cpp b/Engine/Plugins/Experimental/FastGeoStreaming/Source/FastGeoStreaming/Private/FastGeoAsyncRenderStateJobQueue.cpp
--- a/Engine/Plugins/Experimental/FastGeoStreaming/Source/FastGeoStreaming/Private/FastGeoAsyncRenderStateJobQueue.cpp
+++ b/Engine/Plugins/Experimental/FastGeoStreaming/Source/FastGeoStreaming/Private/FastGeoAsyncRenderStateJobQueue.cpp
@@ -37,6 +37,9 @@ void FFastGeoAsyncRenderStateJobQueue::AddJob(const FJob& Job)
 	{
 		switch (Job.Type)
 		{
+		case EJobType::PrecachePSOs:
+			Job.FastGeo->OnPrecachePSOsBegin_GameThread();
+			break;
 		case EJobType::CreateRenderState:
 			Job.FastGeo->OnCreateRenderStateBegin_GameThread();
 			break;
@@ -173,10 +176,16 @@ void FFastGeoAsyncRenderStateJobQueue::FJobSet::OnPostExecute_GameThread() const
 
 void FFastGeoAsyncRenderStateJobQueue::FJob::Execute() const
 {
-	if (const bool bPendingKill = FastGeo && !::IsValid(FastGeo); !bPendingKill)
+	if (::IsValid(FastGeo))
 	{
 		switch (Type)
 		{
+		case EJobType::PrecachePSOs:
+		{
+			TRACE_CPUPROFILER_EVENT_SCOPE(FFastGeoAsyncRenderStateJobQueue::PrecachePSOs);
+			FastGeo->OnPrecachePSOs_Concurrent();
+		}
+		break;
 		case EJobType::CreateRenderState:
 		{
 			TRACE_CPUPROFILER_EVENT_SCOPE(FFastGeoAsyncRenderStateJobQueue::CreateRenderState);
@@ -195,10 +204,13 @@ void FFastGeoAsyncRenderStateJobQueue::FJob::Execute() const
 
 void FFastGeoAsyncRenderStateJobQueue::FJob::OnPostExecute_GameThread() const
 {
-	if (const bool bPendingKill = FastGeo && !::IsValid(FastGeo); !bPendingKill)
+	if (::IsValid(FastGeo))
 	{
 		switch (Type)
 		{
+		case EJobType::PrecachePSOs:
+			FastGeo->OnPrecachePSOsEnd_GameThread();
+			break;
 		case EJobType::CreateRenderState:
 			FastGeo->OnCreateRenderStateEnd_GameThread();
 			break;
diff --git a/Engine/Plugins/Experimental/FastGeoStreaming/Source/FastGeoStreaming/Private/FastGeoAsyncRenderStateJobQueue.h b/Engine/Plugins/Experimental/FastGeoStreaming/Source/FastGeoStreaming/Private/FastGeoAsyncRenderStateJobQueue.h
--- a/Engine/Plugins/Experimental/FastGeoStreaming/Source/FastGeoStreaming/Private/FastGeoAsyncRenderStateJobQueue.h
+++ b/Engine/Plugins/Experimental/FastGeoStreaming/Source/FastGeoStreaming/Private/FastGeoAsyncRenderStateJobQueue.h
@@ -15,6 +15,7 @@ struct FFastGeoAsyncRenderStateJobQueue
 
 	enum class EJobType
 	{
+		PrecachePSOs,
 		CreateRenderState,
 		DestroyRenderState
 	};
diff --git a/Engine/Plugins/Experimental/FastGeoStreaming/Source/FastGeoStreaming/Private/FastGeoComponent.cpp b/Engine/Plugins/Experimental/FastGeoStreaming/Source/FastGeoStreaming/Private/FastGeoComponent.cpp
--- a/Engine/Plugins/Experimental/FastGeoStreaming/Source/FastGeoStreaming/Private/FastGeoComponent.cpp
+++ b/Engine/Plugins/Experimental/FastGeoStreaming/Source/FastGeoStreaming/Private/FastGeoComponent.cpp
@@ -173,4 +173,11 @@ bool FFastGeoComponent::ShouldComponentAddToRenderScene() const
 
 	// If the detail mode setting allows it, add it to the scene.
 	return DetailMode <= GetCachedScalabilityCVars().DetailMode;
-}
\ No newline at end of file
+}
+
+#if !WITH_EDITOR && UE_WITH_PSO_PRECACHING
+bool FFastGeoComponent::ShouldPrecachePSOs() const
+{
+	return FApp::CanEverRender() && IsComponentPSOPrecachingEnabled() && !GetPSOPrecacheComponentData().IsPSOPrecacheCalled();
+}
+#endif
\ No newline at end of file
diff --git a/Engine/Plugins/Experimental/FastGeoStreaming/Source/FastGeoStreaming/Private/FastGeoContainer.cpp b/Engine/Plugins/Experimental/FastGeoStreaming/Source/FastGeoStreaming/Private/FastGeoContainer.cpp
--- a/Engine/Plugins/Experimental/FastGeoStreaming/Source/FastGeoStreaming/Private/FastGeoContainer.cpp
+++ b/Engine/Plugins/Experimental/FastGeoStreaming/Source/FastGeoStreaming/Private/FastGeoContainer.cpp
@@ -216,6 +216,9 @@ void UFastGeoContainer::Register()
 		UFastGeoWorldSubsystem* WorldSubsystem = GetWorld()->GetSubsystemChecked<UFastGeoWorldSubsystem>();
 		if (!PendingCreate.RenderState.ComponentsToProcess.IsEmpty())
 		{
+#if !WITH_EDITOR && UE_WITH_PSO_PRECACHING
+			check(bPrecachedPSOs);
+#endif
 			WorldSubsystem->PushAsyncCreateRenderStateJob(this);
 		}
 
@@ -644,11 +647,7 @@ void UFastGeoContainer::BeginDestroy()
 
 bool UFastGeoContainer::IsReadyForFinishDestroy()
 {
-	bool bResult = Super::IsReadyForFinishDestroy() && DestroyFence.IsFenceComplete();
-#if !WITH_EDITOR && UE_WITH_PSO_PRECACHING
-	bResult = bResult && PrecachePSOsTask.IsCompleted();
-#endif
-	return bResult;
+	return Super::IsReadyForFinishDestroy() && DestroyFence.IsFenceComplete();
 }
 
 void UFastGeoContainer::AddComponentCluster(FFastGeoComponentCluster* ComponentCluster)
@@ -782,39 +781,37 @@ void UFastGeoContainer::SerializeComponentClusters(FArchive& Ar)
 	Ar << HLODs;
 }
 
-void UFastGeoContainer::PrecachePSOs()
+void UFastGeoContainer::OnPrecachePSOs_Concurrent()
 {
-	TRACE_CPUPROFILER_EVENT_SCOPE(UFastGeoContainer::PrecachePSOs_GameThread);
-
 #if !WITH_EDITOR && UE_WITH_PSO_PRECACHING
-	if (bPrecachedPSOs)
-	{
-		return;
-	}
-	bPrecachedPSOs = true;
+	TRACE_CPUPROFILER_EVENT_SCOPE(UFastGeoContainer::OnPrecachePSOs_Concurrent);
 
-	if (IsComponentPSOPrecachingEnabled())
+	check(IsComponentPSOPrecachingEnabled());
+
+	ForEachComponentCluster([](FFastGeoComponentCluster& ComponentCluster)
 	{
-		TArray<FFastGeoComponent*> Components;
-		ForEachComponentCluster([&Components](FFastGeoComponentCluster& ComponentCluster)
+		ComponentCluster.ForEachComponent<FFastGeoComponent>([](FFastGeoComponent& Component)
 		{
-			ComponentCluster.ForEachComponent<FFastGeoComponent>([&Components](FFastGeoComponent& Component)
-			{
-				// Mark component so that IsPSOPrecaching() return true even if component's PSO task has not started
-				Component.MarkPrecachePSOsRequired();
-				Components.Add(&Component);
-			});
+			Component.PrecachePSOs_Concurrent();
 		});
+	});
+#endif
+}
+
+void UFastGeoContainer::PrecachePSOs()
+{
+	TRACE_CPUPROFILER_EVENT_SCOPE(UFastGeoContainer::PrecachePSOs_GameThread);
 
-		check(PrecachePSOsTask.IsCompleted());
-		PrecachePSOsTask = FastGeo::FAsyncTaskWithAssetRefs::Launch(TEXT("UFastGeoContainer::PrecachePSOs_Task"), Assets, [this, Components = MoveTemp(Components)]()
+#if !WITH_EDITOR && UE_WITH_PSO_PRECACHING
+	if (!bPrecachedPSOs)
+	{
+		bPrecachedPSOs = true;
+
+		if (IsComponentPSOPrecachingEnabled())
 		{
-			TRACE_CPUPROFILER_EVENT_SCOPE(UFastGeoContainer::PrecachePSOs_Task);
-			for (FFastGeoComponent* Component : Components)
-			{
-				Component->PrecachePSOs();
-			}
-		}, LowLevelTasks::ETaskPriority::BackgroundNormal);
+			UFastGeoWorldSubsystem* WorldSubsystem = GetWorld()->GetSubsystemChecked<UFastGeoWorldSubsystem>();
+			WorldSubsystem->PushAsyncPrecachePSOsJob(this);
+		}
 	}
 #endif
 }
diff --git a/Engine/Plugins/Experimental/FastGeoStreaming/Source/FastGeoStreaming/Private/FastGeoDecalComponent.cpp b/Engine/Plugins/Experimental/FastGeoStreaming/Source/FastGeoStreaming/Private/FastGeoDecalComponent.cpp
--- a/Engine/Plugins/Experimental/FastGeoStreaming/Source/FastGeoStreaming/Private/FastGeoDecalComponent.cpp
+++ b/Engine/Plugins/Experimental/FastGeoStreaming/Source/FastGeoStreaming/Private/FastGeoDecalComponent.cpp
@@ -5,7 +5,7 @@
 #include "FastGeoWorldSubsystem.h"
 #include "FastGeoRegisteredComponent.h"
 #include "FastGeoDestroyRenderStateContext.h"
-#include "FastGeoPSOPrecacheFinishedTask.h"
+#include "FastGeoPSOWeakPtrTraits.h"
 #include "Components/DecalComponent.h"
 #include "SceneProxies/DeferredDecalProxy.h"
 #include "Materials/MaterialInterface.h"
@@ -29,10 +29,6 @@ FFastGeoDecalComponent::FFastGeoDecalComponent(const FFastGeoDecalComponent& Oth
 	, bIsVisible(Other.bIsVisible)
 	, SceneProxyDesc(Other.SceneProxyDesc)
 	, SceneProxy(nullptr)
-#if !WITH_EDITOR && UE_WITH_PSO_PRECACHING
-	, LatestPSOPrecacheJobSetCompleted(0)
-	, LatestPSOPrecacheJobSet(0)
-#endif
 {
 }
 
@@ -169,7 +165,7 @@ FDeferredDecalProxy* FFastGeoDecalComponent::CreateSceneProxy(ESceneProxyCreatio
 	check(!SceneProxy);
 
 #if !WITH_EDITOR && UE_WITH_PSO_PRECACHING
-	if (LatestPSOPrecacheJobSetCompleted != LatestPSOPrecacheJobSet && GetPSOPrecacheProxyCreationStrategy() == EPSOPrecacheProxyCreationStrategy::DelayUntilPSOPrecached)
+	if (PSOPrecacheComponentData.IsPSOPrecaching() && GetPSOPrecacheProxyCreationStrategy() != EPSOPrecacheProxyCreationStrategy::AlwaysCreate)
 	{
 		OutError = ESceneProxyCreationError::WaitingPSOs;
 		return nullptr;
@@ -241,50 +237,37 @@ void FFastGeoDecalComponent::SetLifeSpan(const float InLifeSpan)
 	}
 }
 
-void FFastGeoDecalComponent::PrecachePSOs()
-{
 #if !WITH_EDITOR && UE_WITH_PSO_PRECACHING
-	TRACE_CPUPROFILER_EVENT_SCOPE(FFastGeoPrimitiveComponent::PrecachePSOs);
+bool FFastGeoDecalComponent::ShouldPrecachePSOs() const
+{
+	return Super::ShouldPrecachePSOs() && SceneProxyDesc.DecalMaterial && !SceneProxyDesc.DecalMaterial->HasAnyFlags(RF_NeedPostLoad);
+}
 
-	if (!FApp::CanEverRender() || !IsComponentPSOPrecachingEnabled())
+void FFastGeoDecalComponent::PrecachePSOs_Concurrent()
+{
+	TRACE_CPUPROFILER_EVENT_SCOPE(FFastGeoDecalComponent::PrecachePSOs_Concurrent);
+	if (!ShouldPrecachePSOs())
 	{
 		return;
 	}
 
-	if (SceneProxyDesc.DecalMaterial && !SceneProxyDesc.DecalMaterial->HasAnyFlags(RF_NeedPostLoad))
-	{
-		FPSOPrecacheParams PSOPrecacheParams;
-		FPSOPrecacheVertexFactoryDataList VertexFactoryDataList;
-		VertexFactoryDataList.Add(FPSOPrecacheVertexFactoryData(&FLocalVertexFactory::StaticType));
+	FPSOPrecacheParams PSOPrecacheParams;
+	FPSOPrecacheVertexFactoryDataList VertexFactoryDataList;
+	VertexFactoryDataList.Add(FPSOPrecacheVertexFactoryData(&FLocalVertexFactory::StaticType));
 
-		// Immediately create at high priority and thus doesn't need boosting anymore
-		TArray<FMaterialPSOPrecacheRequestID> MaterialPSOPrecacheRequestIDs;
-		FGraphEventArray GraphEvents = SceneProxyDesc.DecalMaterial->PrecachePSOs(VertexFactoryDataList, PSOPrecacheParams, EPSOPrecachePriority::High, MaterialPSOPrecacheRequestIDs);
+	// Immediately create at high priority and thus doesn't need boosting anymore
+	TArray<FMaterialPSOPrecacheRequestID> MaterialPSOPrecacheRequestIDs;
+	FGraphEventArray PSOPrecacheCompileEvents = SceneProxyDesc.DecalMaterial->PrecachePSOs(VertexFactoryDataList, PSOPrecacheParams, EPSOPrecachePriority::High, MaterialPSOPrecacheRequestIDs);
 
-		// Request recreate of the render state when the PSO compilation is ready (if we want to delay proxy creation)
-		if (GetPSOPrecacheProxyCreationStrategy() != EPSOPrecacheProxyCreationStrategy::AlwaysCreate)
+	// Request recreate of the render state when the PSO compilation is ready (if we want to delay proxy creation)
+	auto OnPrecacheFinishedGT = [this]()
+	{
+		if (GetWorld())
 		{
-			LatestPSOPrecacheJobSet++;
-			if (GraphEvents.Num() > 0)
-			{
-				TGraphTask<FFastGeoPSOPrecacheFinishedTask>::CreateTask(&GraphEvents).ConstructAndDispatchWhenReady(this, LatestPSOPrecacheJobSet);
-			}
-			else
-			{
-				// No graph events to wait on, the job set can be considered complete.
-				LatestPSOPrecacheJobSetCompleted = LatestPSOPrecacheJobSet;
-			}
+			MarkRenderStateDirty(false);
 		}
-	}
-#endif
-}
+	};
 
-#if !WITH_EDITOR && UE_WITH_PSO_PRECACHING
-void FFastGeoDecalComponent::OnPrecacheFinished(int32 JobSetThatJustCompleted)
-{
-	QUICK_SCOPE_CYCLE_COUNTER(STAT_PSOPrecacheFinishedTask);
-	int32 CurrJobSetCompleted = LatestPSOPrecacheJobSetCompleted.load();
-	while (CurrJobSetCompleted < JobSetThatJustCompleted && !LatestPSOPrecacheJobSetCompleted.compare_exchange_weak(CurrJobSetCompleted, JobSetThatJustCompleted)) {}
-	MarkRenderStateDirty(false);
+	PSOPrecacheComponentData.SetMaterialPSOPrecacheData(this, MaterialPSOPrecacheRequestIDs, PSOPrecacheCompileEvents, OnPrecacheFinishedGT);
 }
 #endif
\ No newline at end of file
diff --git a/Engine/Plugins/Experimental/FastGeoStreaming/Source/FastGeoStreaming/Private/FastGeoLocalLightComponent.cpp b/Engine/Plugins/Experimental/FastGeoStreaming/Source/FastGeoStreaming/Private/FastGeoLocalLightComponent.cpp
--- a/Engine/Plugins/Experimental/FastGeoStreaming/Source/FastGeoStreaming/Private/FastGeoLocalLightComponent.cpp
+++ b/Engine/Plugins/Experimental/FastGeoStreaming/Source/FastGeoStreaming/Private/FastGeoLocalLightComponent.cpp
@@ -2,6 +2,7 @@
 
 #include "FastGeoLocalLightComponent.h"
 #include "FastGeoWorldSubsystem.h"
+#include "FastGeoPSOWeakPtrTraits.h"
 #include "Components/LocalLightComponent.h"
 #include "PointLightSceneProxy.h"
 #include "PointLightSceneProxyDesc.h"
@@ -21,6 +22,15 @@ FFastGeoLocalLightComponent::FFastGeoLocalLightComponent(int32 InComponentIndex,
 {
 }
 
+FFastGeoLocalLightComponent::FFastGeoLocalLightComponent(const FFastGeoLocalLightComponent& Other)
+	: Super(Other)
+	, bAffectsWorld(Other.bAffectsWorld)
+	, bSourceComponentIsVisible(Other.bSourceComponentIsVisible)
+	, bIsVisible(true)
+	, SceneProxy(nullptr)
+{
+}
+
 #if WITH_EDITOR
 void FFastGeoLocalLightComponent::InitializeFromComponent(UActorComponent* Component)
 {
@@ -71,7 +81,8 @@ void FFastGeoLocalLightComponent::CreateRenderState(FRegisterComponentContext* C
 	bRenderStateDirty = false;
 
 	check(!SceneProxy);
-	SceneProxy = CreateSceneProxy();
+	ESceneProxyCreationError Error;
+	SceneProxy = CreateSceneProxy(Error);
 	if (SceneProxy)
 	{
 		check(bIsVisible);
@@ -84,16 +95,29 @@ void FFastGeoLocalLightComponent::CreateRenderState(FRegisterComponentContext* C
 		check(LightSceneDesc.SceneProxy);
 		ProxyState = EProxyCreationState::Created;
 	}
+	else if (Error == ESceneProxyCreationError::WaitingPSOs)
+	{
+		ProxyState = EProxyCreationState::Delayed;
+	}
 	else
 	{
 		ProxyState = EProxyCreationState::None;
 	}
 }
 
-FLightSceneProxy* FFastGeoLocalLightComponent::CreateSceneProxy()
+FLightSceneProxy* FFastGeoLocalLightComponent::CreateSceneProxy(ESceneProxyCreationError& OutError)
 {
 	check(ShouldComponentAddToRenderScene());
 
+#if !WITH_EDITOR && UE_WITH_PSO_PRECACHING
+	if (PSOPrecacheComponentData.IsPSOPrecaching() && GetPSOPrecacheProxyCreationStrategy() != EPSOPrecacheProxyCreationStrategy::AlwaysCreate)
+	{
+		OutError = ESceneProxyCreationError::WaitingPSOs;
+		return nullptr;
+	}
+#endif
+
+	OutError = ESceneProxyCreationError::None;
 	FSceneInterface* Scene = GetScene();
 	check(Scene);
 	FLocalLightSceneProxyDesc& SceneProxyDesc = GetLocalLightSceneProxyDesc();
@@ -104,7 +128,6 @@ FLightSceneProxy* FFastGeoLocalLightComponent::CreateSceneProxy()
 void FFastGeoLocalLightComponent::DestroyRenderState(FFastGeoDestroyRenderStateContext* Context)
 {
 	FWriteScopeLock WriteLock(*Lock.Get());
-	check(SceneProxy || ProxyState == EProxyCreationState::None);
 	if (SceneProxy)
 	{
 		check(ProxyState == EProxyCreationState::Created);
@@ -116,6 +139,11 @@ void FFastGeoLocalLightComponent::DestroyRenderState(FFastGeoDestroyRenderStateC
 		ProxyState = EProxyCreationState::None;
 		bRenderStateDirty = false;
 	}
+	else
+	{
+		check(ProxyState == EProxyCreationState::None || ProxyState == EProxyCreationState::Delayed);
+		ProxyState = EProxyCreationState::None;
+	}
 }
 
 void FFastGeoLocalLightComponent::InitializeSceneProxyDescDynamicProperties()
@@ -140,4 +168,46 @@ bool FFastGeoLocalLightComponent::UpdateVisibilityInternal()
 	bIsVisible = true;
 #endif
 	return bIsVisible != bOldIsVisible;
-}
\ No newline at end of file
+}
+
+#if !WITH_EDITOR && UE_WITH_PSO_PRECACHING
+bool FFastGeoLocalLightComponent::ShouldPrecachePSOs() const
+{
+	if (!Super::ShouldPrecachePSOs())
+	{
+		return false;
+	}
+
+	const FLocalLightSceneProxyDesc& SceneProxyDesc = GetLocalLightSceneProxyDesc();
+	return SceneProxyDesc.LightFunctionMaterial && !SceneProxyDesc.LightFunctionMaterial->HasAnyFlags(RF_NeedPostLoad);
+}
+
+void FFastGeoLocalLightComponent::PrecachePSOs_Concurrent()
+{
+	TRACE_CPUPROFILER_EVENT_SCOPE(FFastGeoLocalLightComponent::PrecachePSOs_Concurrent);
+	if (!ShouldPrecachePSOs())
+	{
+		return;
+	}
+
+	FPSOPrecacheParams PSOPrecacheParams;
+	FPSOPrecacheVertexFactoryDataList VertexFactoryDataList;
+	VertexFactoryDataList.Add(FPSOPrecacheVertexFactoryData(&FLocalVertexFactory::StaticType));
+
+	// Immediately create at high priority and thus doesn't need boosting anymore
+	TArray<FMaterialPSOPrecacheRequestID> MaterialPSOPrecacheRequestIDs;
+	FLocalLightSceneProxyDesc& SceneProxyDesc = GetLocalLightSceneProxyDesc();
+	FGraphEventArray PSOPrecacheCompileEvents = SceneProxyDesc.LightFunctionMaterial->PrecachePSOs(VertexFactoryDataList, PSOPrecacheParams, EPSOPrecachePriority::High, MaterialPSOPrecacheRequestIDs);
+
+	// Request recreate of the render state when the PSO compilation is ready (if we want to delay proxy creation)
+	auto OnPrecacheFinishedGT = [this]()
+	{
+		if (GetWorld())
+		{
+			MarkRenderStateDirty(false);
+		}
+	};
+
+	PSOPrecacheComponentData.SetMaterialPSOPrecacheData(this, MaterialPSOPrecacheRequestIDs, PSOPrecacheCompileEvents, OnPrecacheFinishedGT);
+}
+#endif
\ No newline at end of file
diff --git a/Engine/Plugins/Experimental/FastGeoStreaming/Source/FastGeoStreaming/Private/FastGeoPointLightComponent.cpp b/Engine/Plugins/Experimental/FastGeoStreaming/Source/FastGeoStreaming/Private/FastGeoPointLightComponent.cpp
--- a/Engine/Plugins/Experimental/FastGeoStreaming/Source/FastGeoStreaming/Private/FastGeoPointLightComponent.cpp
+++ b/Engine/Plugins/Experimental/FastGeoStreaming/Source/FastGeoStreaming/Private/FastGeoPointLightComponent.cpp
@@ -33,6 +33,13 @@ bool FFastGeoPointLightComponentBase::IsPointLightSupported() const
 	return true;
 }
 
+#if !WITH_EDITOR && UE_WITH_PSO_PRECACHING
+bool FFastGeoPointLightComponentBase::ShouldPrecachePSOs() const
+{
+	return Super::ShouldPrecachePSOs() && IsPointLightSupported();
+}
+#endif
+
 const FFastGeoElementType FFastGeoPointLightComponent::Type(&FFastGeoPointLightComponentBase::Type);
 
 FFastGeoPointLightComponent::FFastGeoPointLightComponent(int32 InComponentIndex, FFastGeoElementType InType)
diff --git a/Engine/Plugins/Experimental/FastGeoStreaming/Source/FastGeoStreaming/Private/FastGeoPrimitiveComponent.cpp b/Engine/Plugins/Experimental/FastGeoStreaming/Source/FastGeoStreaming/Private/FastGeoPrimitiveComponent.cpp
--- a/Engine/Plugins/Experimental/FastGeoStreaming/Source/FastGeoStreaming/Private/FastGeoPrimitiveComponent.cpp
+++ b/Engine/Plugins/Experimental/FastGeoStreaming/Source/FastGeoStreaming/Private/FastGeoPrimitiveComponent.cpp
@@ -9,7 +9,7 @@
 #include "FastGeoDestroyRenderStateContext.h"
 #include "FastGeoWeakElement.h"
 #include "FastGeoLog.h"
-#include "FastGeoPSOPrecacheFinishedTask.h"
+#include "FastGeoPSOWeakPtrTraits.h"
 #include "AI/NavigationModifier.h"
 #include "AI/Navigation/NavigationRelevantData.h"
 #include "PrimitiveSceneDesc.h"
@@ -54,14 +54,6 @@ FFastGeoPrimitiveComponent::FFastGeoPrimitiveComponent(const FFastGeoPrimitiveCo
 	, BodyInstanceOwner()
 	, PrimitiveSceneData()
 	, AsyncTermBodyPayload()
-#if !WITH_EDITOR && UE_WITH_PSO_PRECACHING
-	, MaterialPSOPrecacheRequestIDs()
-	, LatestPSOPrecacheJobSetCompleted(0)
-	, LatestPSOPrecacheJobSet(0)
-	, bPSOPrecacheCalled(false)
-	, bPSOPrecacheRequired(false)
-	, PSOPrecacheRequestPriority(EPSOPrecachePriority::Medium)
-#endif
 {
 }
 
@@ -611,28 +603,32 @@ void FFastGeoPrimitiveComponent::InitializeDynamicProperties()
 #endif
 }
 
+bool FFastGeoPrimitiveComponent::ShouldRenderProxyFallbackToDefaultMaterial() const
+{
 #if !WITH_EDITOR && UE_WITH_PSO_PRECACHING
-void FFastGeoPrimitiveComponent::OnPrecacheFinished(int32 JobSetThatJustCompleted)
+	return PSOPrecacheComponentData.IsPSOPrecaching() && GetPSOPrecacheProxyCreationStrategy() == EPSOPrecacheProxyCreationStrategy::UseDefaultMaterialUntilPSOPrecached;
+#else
+	return false;
+#endif
+}
+
+void FFastGeoPrimitiveComponent::PrecachePSOs()
 {
-	QUICK_SCOPE_CYCLE_COUNTER(STAT_PSOPrecacheFinishedTask);
-	int32 CurrJobSetCompleted = LatestPSOPrecacheJobSetCompleted.load();
-	while (CurrJobSetCompleted < JobSetThatJustCompleted && !LatestPSOPrecacheJobSetCompleted.compare_exchange_weak(CurrJobSetCompleted, JobSetThatJustCompleted)) {}
-	MarkRenderStateDirty();
+#if !WITH_EDITOR && UE_WITH_PSO_PRECACHING
+	PrecachePSOs_Concurrent();
+#endif
 }
 
-void FFastGeoPrimitiveComponent::RequestRecreateRenderStateWhenPSOPrecacheFinished(const FGraphEventArray& PSOPrecacheCompileEvents)
+bool FFastGeoPrimitiveComponent::CheckPSOPrecachingAndBoostPriority(EPSOPrecachePriority NewPSOPrecachePriority)
 {
-	// If the proxy creation strategy relies on knowing when the precached PSO has been compiled,
-	// schedule a task to mark the render state dirty when all PSOs are compiled so the proxy gets recreated.
-	if (GetPSOPrecacheProxyCreationStrategy() != EPSOPrecacheProxyCreationStrategy::AlwaysCreate)
-	{
-		LatestPSOPrecacheJobSet++;
-		// Even if PSOPrecacheCompileEvents is empty, still push the completion task as it needs to run on the Game Thread and call MarkRenderStateDirty
-		TGraphTask<FFastGeoPSOPrecacheFinishedTask>::CreateTask(&PSOPrecacheCompileEvents).ConstructAndDispatchWhenReady(this, LatestPSOPrecacheJobSet);
-	}
-	bPSOPrecacheCalled = true;
+#if !WITH_EDITOR && UE_WITH_PSO_PRECACHING
+	return PSOPrecacheComponentData.CheckPSOPrecachingAndBoostPriority(NewPSOPrecachePriority);
+#else
+	return false;
+#endif
 }
 
+#if !WITH_EDITOR && UE_WITH_PSO_PRECACHING
 void FFastGeoPrimitiveComponent::SetupPrecachePSOParams(FPSOPrecacheParams& Params)
 {
 	const FPrimitiveSceneProxyDesc& SceneProxyDesc = GetSceneProxyDesc();
@@ -662,94 +658,37 @@ void FFastGeoPrimitiveComponent::SetupPrecachePSOParams(FPSOPrecacheParams& Para
 	}
 }
 
-#endif
-
-bool FFastGeoPrimitiveComponent::IsPSOPrecaching() const
+void FFastGeoPrimitiveComponent::PrecachePSOs_Concurrent()
 {
-#if !WITH_EDITOR && UE_WITH_PSO_PRECACHING
-	// Consider as precaching when marked as required to do PSOs precaching (even if task has not been launched yet) 
-	return bPSOPrecacheRequired || (LatestPSOPrecacheJobSetCompleted != LatestPSOPrecacheJobSet);
-#else
-	return false;
-#endif
-}
-
-bool FFastGeoPrimitiveComponent::ShouldRenderProxyFallbackToDefaultMaterial() const
-{
-#if !WITH_EDITOR && UE_WITH_PSO_PRECACHING
-	return IsPSOPrecaching() && GetPSOPrecacheProxyCreationStrategy() == EPSOPrecacheProxyCreationStrategy::UseDefaultMaterialUntilPSOPrecached;
-#else
-	return false;
-#endif
-}
-
-bool FFastGeoPrimitiveComponent::CheckPSOPrecachingAndBoostPriority(EPSOPrecachePriority NewPSOPrecachePriority)
-{
-#if !WITH_EDITOR && UE_WITH_PSO_PRECACHING
-	bool bPrecacheStillRunning = IsPSOPrecaching();
-
-	ensure(!IsComponentPSOPrecachingEnabled() || bPSOPrecacheCalled || bPSOPrecacheRequired);
-	check(NewPSOPrecachePriority == EPSOPrecachePriority::High || NewPSOPrecachePriority == EPSOPrecachePriority::Highest);
-
-	if (bPrecacheStillRunning && PSOPrecacheRequestPriority < NewPSOPrecachePriority)
-	{
-		// Only boost PSO priority if PSO task was started
-		if (LatestPSOPrecacheJobSetCompleted != LatestPSOPrecacheJobSet)
-		{
-			BoostPSOPriority(NewPSOPrecachePriority, MaterialPSOPrecacheRequestIDs);
-		}
-		PSOPrecacheRequestPriority = NewPSOPrecachePriority;
-	}
-	return bPrecacheStillRunning;
-#else
-	return false;
-#endif
-}
-
-
-void FFastGeoPrimitiveComponent::MarkPrecachePSOsRequired()
-{
-#if !WITH_EDITOR && UE_WITH_PSO_PRECACHING
-	if (bPSOPrecacheCalled || !FApp::CanEverRender() || !IsComponentPSOPrecachingEnabled())
+	TRACE_CPUPROFILER_EVENT_SCOPE(FFastGeoPrimitiveComponent::PrecachePSOs_Concurrent);
+	if (!ShouldPrecachePSOs())
 	{
 		return;
 	}
-	bPSOPrecacheRequired = true;
-	PSOPrecacheRequestPriority = EPSOPrecachePriority::Medium;
-#endif
-}
-
-void FFastGeoPrimitiveComponent::PrecachePSOs()
-{
-#if !WITH_EDITOR && UE_WITH_PSO_PRECACHING
-	TRACE_CPUPROFILER_EVENT_SCOPE(FFastGeoPrimitiveComponent::PrecachePSOs);
-
-	if (!bPSOPrecacheRequired)
-	{
-		return;
-	}
-
-	// Clear the current request data
-	MaterialPSOPrecacheRequestIDs.Empty();
 
 	// Collect the data from the derived classes
 	FPSOPrecacheParams PSOPrecacheParams;
 	SetupPrecachePSOParams(PSOPrecacheParams);
+
 	FMaterialInterfacePSOPrecacheParamsList PSOPrecacheDataArray;
 	CollectPSOPrecacheData(PSOPrecacheParams, PSOPrecacheDataArray);
-	// Set priority
-	for (FMaterialInterfacePSOPrecacheParams& Params : PSOPrecacheDataArray)
-	{
-		Params.Priority = PSOPrecacheRequestPriority;
-	}
 
-	FGraphEventArray GraphEvents;
-	PrecacheMaterialPSOs(PSOPrecacheDataArray, MaterialPSOPrecacheRequestIDs, GraphEvents);
+	FGraphEventArray PSOPrecacheCompileEvents;
+	TArray<FMaterialPSOPrecacheRequestID> MaterialPSOPrecacheRequestIDs;
+	PrecacheMaterialPSOs(PSOPrecacheDataArray, MaterialPSOPrecacheRequestIDs, PSOPrecacheCompileEvents);
 
-	RequestRecreateRenderStateWhenPSOPrecacheFinished(GraphEvents);
-	bPSOPrecacheRequired = false;
-#endif
+	// Request recreate of the render state when the PSO compilation is ready (if we want to delay proxy creation)
+	auto OnPrecacheFinishedGT = [this]()
+	{
+		if (GetWorld())
+		{
+			MarkRenderStateDirty();
+		}
+	};
+
+	PSOPrecacheComponentData.SetMaterialPSOPrecacheData(this, MaterialPSOPrecacheRequestIDs, PSOPrecacheCompileEvents, OnPrecacheFinishedGT);
 }
+#endif
 
 UObject* FFastGeoPrimitiveComponent::GetOwner() const
 {
diff --git a/Engine/Plugins/Experimental/FastGeoStreaming/Source/FastGeoStreaming/Private/FastGeoWorldSubsystem.cpp b/Engine/Plugins/Experimental/FastGeoStreaming/Source/FastGeoStreaming/Private/FastGeoWorldSubsystem.cpp
--- a/Engine/Plugins/Experimental/FastGeoStreaming/Source/FastGeoStreaming/Private/FastGeoWorldSubsystem.cpp
+++ b/Engine/Plugins/Experimental/FastGeoStreaming/Source/FastGeoStreaming/Private/FastGeoWorldSubsystem.cpp
@@ -569,24 +569,34 @@ void UFastGeoWorldSubsystem::PushAsyncDestroyContainer(UFastGeoContainer* FastGe
 	}
 }
 
-void UFastGeoWorldSubsystem::PushAsyncCreateRenderStateJob(UFastGeoContainer* FastGeo)
+FFastGeoAsyncRenderStateJobQueue* UFastGeoWorldSubsystem::GetOrCreateJobQueue()
 {
-	check(IsInGameThread());
 	if (!AsyncRenderStateJobQueue.IsValid())
 	{
-		AsyncRenderStateJobQueue = MakeUnique<FFastGeoAsyncRenderStateJobQueue>(FastGeo->GetWorld()->Scene);
+		AsyncRenderStateJobQueue = MakeUnique<FFastGeoAsyncRenderStateJobQueue>(GetWorld()->Scene);
 	}
-	AsyncRenderStateJobQueue->AddJob({ FastGeo, FFastGeoAsyncRenderStateJobQueue::EJobType::CreateRenderState });
+	return AsyncRenderStateJobQueue.Get();
+}
+
+void UFastGeoWorldSubsystem::PushAsyncPrecachePSOsJob(UFastGeoContainer* FastGeo)
+{
+	check(IsInGameThread());
+	check(FastGeo->GetWorld() == GetWorld());
+	GetOrCreateJobQueue()->AddJob({FastGeo, FFastGeoAsyncRenderStateJobQueue::EJobType::PrecachePSOs});
+}
+
+void UFastGeoWorldSubsystem::PushAsyncCreateRenderStateJob(UFastGeoContainer* FastGeo)
+{
+	check(IsInGameThread());
+	check(FastGeo->GetWorld() == GetWorld());
+	GetOrCreateJobQueue()->AddJob({ FastGeo, FFastGeoAsyncRenderStateJobQueue::EJobType::CreateRenderState });
 }
 
 void UFastGeoWorldSubsystem::PushAsyncDestroyRenderStateJob(UFastGeoContainer* FastGeo)
 {
 	check(IsInGameThread());
-	if (!AsyncRenderStateJobQueue.IsValid())
-	{
-		AsyncRenderStateJobQueue = MakeUnique<FFastGeoAsyncRenderStateJobQueue>(FastGeo->GetWorld()->Scene);
-	}
-	AsyncRenderStateJobQueue->AddJob({ FastGeo, FFastGeoAsyncRenderStateJobQueue::EJobType::DestroyRenderState });
+	check(FastGeo->GetWorld() == GetWorld());
+	GetOrCreateJobQueue()->AddJob({ FastGeo, FFastGeoAsyncRenderStateJobQueue::EJobType::DestroyRenderState });
 }
 
 void UFastGeoWorldSubsystem::ProcessAsyncRenderStateJobs(bool bWaitForCompletion)
diff --git a/Engine/Plugins/Experimental/FastGeoStreaming/Source/FastGeoStreaming/Private/FastGeoWorldSubsystem.h b/Engine/Plugins/Experimental/FastGeoStreaming/Source/FastGeoStreaming/Private/FastGeoWorldSubsystem.h
--- a/Engine/Plugins/Experimental/FastGeoStreaming/Source/FastGeoStreaming/Private/FastGeoWorldSubsystem.h
+++ b/Engine/Plugins/Experimental/FastGeoStreaming/Source/FastGeoStreaming/Private/FastGeoWorldSubsystem.h
@@ -41,6 +41,12 @@ public:
 	void AddToComponentsPendingRecreate(FFastGeoComponent* InComponentPendingRecreate);
 	void ProcessPendingRecreate();
 
+	/**
+	 * Start a job that will precache PSOs asynchronously for a specified FastGeo container.
+	 * @param FastGeo: FastGeo container.
+	 */
+	void PushAsyncPrecachePSOsJob(UFastGeoContainer* FastGeo);
+
 	/**
 	 * Start a job that will create render state asynchronously for a specified FastGeo container.
 	 * @param FastGeo: FastGeo container.
@@ -102,6 +108,8 @@ private:
 
 	void PushAsyncDestroyContainer(UFastGeoContainer* FastGeo);
 
+	FFastGeoAsyncRenderStateJobQueue* GetOrCreateJobQueue();
+
 	UPROPERTY()
 	TSet<TObjectPtr<UFastGeoContainer>> PendingDestroyContainers;
 
