[FastGeo] Allow transformation of stationary lights to properly support light function materials.
#Synced-CL 48520266

diff --git a/Engine/Plugins/Experimental/FastGeoStreaming/Source/FastGeoStreaming/Private/FastGeoWorldPartitionRuntimeCellTransformer.cpp b/Engine/Plugins/Experimental/FastGeoStreaming/Source/FastGeoStreaming/Private/FastGeoWorldPartitionRuntimeCellTransformer.cpp
--- a/Engine/Plugins/Experimental/FastGeoStreaming/Source/FastGeoStreaming/Private/FastGeoWorldPartitionRuntimeCellTransformer.cpp
+++ b/Engine/Plugins/Experimental/FastGeoStreaming/Source/FastGeoStreaming/Private/FastGeoWorldPartitionRuntimeCellTransformer.cpp
@@ -739,6 +739,35 @@ FFastGeoTransformResult UFastGeoWorldPartitionRuntimeCellTransformer::IsAllowedA
 	return { EFastGeoTransform::Reject, [InActor, ActorClass] { return FString::Printf(TEXT("Actor %s class is an unsupported class (%s)"), *InActor->GetName(), *ActorClass->GetName()); } };
 }
 
+bool UFastGeoWorldPartitionRuntimeCellTransformer::IsComponentMobilitySupported(USceneComponent* InComponent, FString& OutReason) const
+{
+	bool bIsSupported = false;
+	check(InComponent);
+
+	const EComponentMobility::Type Mobility = InComponent->GetMobility();
+	if (InComponent->IsA<UDecalComponent>())
+	{
+		// Allow all Mobility values for DecalComponent, as its mobility is Movable
+		bIsSupported = true;
+	}
+	else if (InComponent->IsA<UPointLightComponent>() || InComponent->IsA<URectLightComponent>() || InComponent->IsA<USpotLightComponent>())
+	{
+		// Allow Stationary support LightComponents as Light Function Materials requires at least Stationary
+		bIsSupported = (Mobility != EComponentMobility::Movable);
+	}
+	else
+	{
+		bIsSupported = (Mobility == EComponentMobility::Static);
+	}
+
+	if (!bIsSupported)
+	{
+		OutReason = FString::Printf(TEXT("(%s) has unsupported Mobility: %s"), *InComponent->GetClass()->GetName(), *StaticEnum<EComponentMobility::Type>()->GetNameStringByValue((int64)Mobility));
+	}
+
+	return bIsSupported;
+}
+
 FFastGeoTransformResult UFastGeoWorldPartitionRuntimeCellTransformer::CanTransformActor(AActor* InActor, bool& bOutIsActorFullyTransformable, TArray<UActorComponent*>& OutTransformableComponents) const
 {
 	bOutIsActorFullyTransformable = false;
@@ -785,11 +814,15 @@ FFastGeoTransformResult UFastGeoWorldPartitionRuntimeCellTransformer::CanTransfo
 		return { EFastGeoTransform::Reject, [InActor] { return FString::Printf(TEXT("Actor %s is replicated"), *InActor->GetName()); } };
 	}
 
-	// Ignore DecalComponent for now, as its mobility is Movable
-	const bool bIsRootDecalComponent = InActor->GetRootComponent() && InActor->GetRootComponent()->IsA<UDecalComponent>();
-	if (!InActor->IsRootComponentStatic() && !bIsRootDecalComponent)
+	if (!InActor->GetRootComponent())
 	{
-		return { EFastGeoTransform::Reject, [InActor] { return FString::Printf(TEXT("Actor %s RootComponent Mobility is not Static"), *InActor->GetName()); } };
+		return { EFastGeoTransform::Reject, [InActor] { return FString::Printf(TEXT("Actor %s RootComponent is invalid"), *InActor->GetName()); } };
+	}
+
+	Reason.Reset();
+	if (!IsComponentMobilitySupported(InActor->GetRootComponent(), Reason))
+	{
+		return { EFastGeoTransform::Reject, [InActor, Reason] { return FString::Printf(TEXT("Actor %s RootComponent %s"), *InActor->GetName(), *Reason); } };
 	}
 
 	if (InActor->IsEditorOnly())
@@ -939,10 +972,10 @@ FFastGeoTransformResult UFastGeoWorldPartitionRuntimeCellTransformer::CanTransfo
 		return { EFastGeoTransform::Reject, [SceneComponent, Reason] { return FString::Printf(TEXT("Component %s [%s]"), *FastGeo::GetComponentShortName(SceneComponent), *Reason); } };
 	}
 
-	// Ignore DecalComponent for now, as its mobility is Movable
-	if (SceneComponent->Mobility != EComponentMobility::Static && !SceneComponent->IsA<UDecalComponent>())
+	Reason.Reset();
+	if (!IsComponentMobilitySupported(SceneComponent, Reason))
 	{
-		return { EFastGeoTransform::Reject, [SceneComponent] { return FString::Printf(TEXT("Component %s Mobility is not Static"), *FastGeo::GetComponentShortName(SceneComponent)); } };
+		return { EFastGeoTransform::Reject, [SceneComponent, Reason] { return FString::Printf(TEXT("Component %s %s"), *FastGeo::GetComponentShortName(SceneComponent), *Reason); } };
 	}
 
 	UPrimitiveComponent* PrimitiveComponent = Cast<UPrimitiveComponent>(InComponent);
diff --git a/Engine/Plugins/Experimental/FastGeoStreaming/Source/FastGeoStreaming/Public/FastGeoWorldPartitionRuntimeCellTransformer.h b/Engine/Plugins/Experimental/FastGeoStreaming/Source/FastGeoStreaming/Public/FastGeoWorldPartitionRuntimeCellTransformer.h
--- a/Engine/Plugins/Experimental/FastGeoStreaming/Source/FastGeoStreaming/Public/FastGeoWorldPartitionRuntimeCellTransformer.h
+++ b/Engine/Plugins/Experimental/FastGeoStreaming/Source/FastGeoStreaming/Public/FastGeoWorldPartitionRuntimeCellTransformer.h
@@ -133,6 +133,8 @@ protected:
 	virtual bool IsComponentTransformable(UActorComponent* InComponent, FString& OutReason) const { return true; }
 	/** Whether the fully transformed actor can be deleted by the transformer. */
 	virtual bool IsFullyTransformedActorDeletable(AActor* InActor, FString& OutReason) const { return true; }
+	/** Whether component mobility is supported by the transformer. */
+	bool IsComponentMobilitySupported(USceneComponent* InComponent, FString& OutReason) const;
 
 	bool IsBlueprintActorWithLogic(AActor* InActor) const;
 	FFastGeoTransformResult CanTransformActor(AActor* InActor, bool& bOutIsActorFullyTransformable, TArray<UActorComponent*>& OutTransformableComponents) const;
