[FastGeo] Fixed invalid text for FastGeo actor coloration tooltip.

diff --git a/Engine/Plugins/Experimental/FastGeoStreaming/Source/FastGeoStreaming/Private/FastGeoWorldSubsystem.cpp b/Engine/Plugins/Experimental/FastGeoStreaming/Source/FastGeoStreaming/Private/FastGeoWorldSubsystem.cpp
--- a/Engine/Plugins/Experimental/FastGeoStreaming/Source/FastGeoStreaming/Private/FastGeoWorldSubsystem.cpp
+++ b/Engine/Plugins/Experimental/FastGeoStreaming/Source/FastGeoStreaming/Private/FastGeoWorldSubsystem.cpp
@@ -125,7 +125,7 @@ UFastGeoWorldSubsystem::UFastGeoWorldSubsystem()
 		FActorPrimitiveColorHandler::FPrimitiveColorHandler FastGeoColorHandler;
 		FastGeoColorHandler.HandlerName = NAME_FastGeoColorHandler;
 		FastGeoColorHandler.HandlerText = LOCTEXT("FastGeo", "FastGeo");
-		FastGeoColorHandler.HandlerToolTipText = LOCTEXT("FastGeoColor_ToopTip", "Colorize FastGeo primitives. ISM [Orange], HLOD ISM [Red], SM [Cyan], HLOD SM [Blue], else White.");
+		FastGeoColorHandler.HandlerToolTipText = LOCTEXT("FastGeoColor_ToopTip", "Colorize FastGeo primitives: FastGeo [Blue], Non-FastGeo [Red]");
 		FastGeoColorHandler.bAvailalbleInEditor = false;
 		FastGeoColorHandler.GetColorFunc = [](const UPrimitiveComponent* InPrimitiveComponent)
 		{
