[FastGeo] Moved FFastGeoPhysicsBodyInstanceOwner in its own file.
- Added more functionality in IPhysicsBodyInstanceOwnerResolver: Added GetCollisionObjectType, GetCollisionEnabled, GetTransform, GetPhysicsObjectById, GetAllPhysicsObjects.
- Implemented new methods of IPhysicsBodyInstanceOwner in FFastGeoPhysicsBodyInstanceOwner.

#Synced-CL 49040969

diff --git a/Engine/Plugins/Experimental/FastGeoStreaming/Source/FastGeoStreaming/Internal/FastGeoInstancedStaticMeshComponent.h b/Engine/Plugins/Experimental/FastGeoStreaming/Source/FastGeoStreaming/Internal/FastGeoInstancedStaticMeshComponent.h
--- a/Engine/Plugins/Experimental/FastGeoStreaming/Source/FastGeoStreaming/Internal/FastGeoInstancedStaticMeshComponent.h
+++ b/Engine/Plugins/Experimental/FastGeoStreaming/Source/FastGeoStreaming/Internal/FastGeoInstancedStaticMeshComponent.h
@@ -33,6 +33,8 @@ public:
 	virtual bool IsNavigationRelevant() const override;
 	virtual bool DoCustomNavigableGeometryExport(FNavigableGeometryExport& GeomExport) const override;
 	virtual bool ShouldSkipNavigationDirtyAreaOnAddOrRemove() override { return true; }
+	virtual Chaos::FPhysicsObject* GetPhysicsObjectById(Chaos::FPhysicsObjectId Id) const override;
+	virtual TArray<Chaos::FPhysicsObject*> GetAllPhysicsObjects() const override;
 	//~ End FFastGeoPrimitiveComponent interface
 
 protected:
diff --git a/Engine/Plugins/Experimental/FastGeoStreaming/Source/FastGeoStreaming/Internal/FastGeoPhysicsBodyInstanceOwner.h b/Engine/Plugins/Experimental/FastGeoStreaming/Source/FastGeoStreaming/Internal/FastGeoPhysicsBodyInstanceOwner.h
new file mode 100644
--- /dev/null
+++ b/Engine/Plugins/Experimental/FastGeoStreaming/Source/FastGeoStreaming/Internal/FastGeoPhysicsBodyInstanceOwner.h
@@ -0,0 +1,47 @@
+// Copyright Epic Games, Inc. All Rights Reserved.
+
+#pragma once
+
+#include "Chaos/ChaosUserEntity.h"
+#include "PhysicsEngine/PhysicsBodyInstanceOwnerInterface.h"
+
+class FFastGeoPrimitiveComponent;
+class UFastGeoContainer;
+
+class FFastGeoPhysicsBodyInstanceOwner : public FChaosUserDefinedEntity, public IPhysicsBodyInstanceOwner
+{
+public:
+	FFastGeoPhysicsBodyInstanceOwner();
+	virtual ~FFastGeoPhysicsBodyInstanceOwner() = default;
+
+	/** Returns the IPhysicsBodyInstanceOwner based on the provided FChaosUserDefinedEntity */
+	static IPhysicsBodyInstanceOwner* GetPhysicsBodyInstanceOwner(FChaosUserDefinedEntity* InUserDefinedEntity);
+
+	//~ Begin FChaosUserDefinedEntity interface
+	TWeakObjectPtr<UObject> GetOwnerObject() override;
+	//~ End FChaosUserDefinedEntity interface
+
+	//~ Begin IPhysicsBodyInstanceOwner interface
+	virtual bool IsStaticPhysics() const override;
+	virtual UObject* GetSourceObject() const override;
+	virtual FTransform GetTransform() const override;
+	virtual ECollisionChannel GetCollisionObjectType() const override;
+	virtual ECollisionEnabled::Type GetCollisionEnabled() const override;
+	virtual ECollisionResponse GetCollisionResponseToChannel(ECollisionChannel Channel) const override;
+	virtual UPhysicalMaterial* GetPhysicalMaterial() const override;
+	virtual void GetComplexPhysicalMaterials(TArray<UPhysicalMaterial*>& OutPhysMaterials, TArray<FPhysicalMaterialMaskParams>* OutPhysMaterialMasks) const override;
+	virtual Chaos::FPhysicsObject* GetPhysicsObjectById(Chaos::FPhysicsObjectId Id) const override;
+	virtual TArray<Chaos::FPhysicsObject*> GetAllPhysicsObjects() const override;
+	//~ End IPhysicsBodyInstanceOwner interface
+
+private:
+	void Uninitialize();
+	void Initialize(FFastGeoPrimitiveComponent* InOwner);
+
+	FFastGeoPrimitiveComponent* OwnerComponent;
+	TWeakObjectPtr<UFastGeoContainer> OwnerContainer;
+	static const FName NAME_FastGeoPhysicsBodyInstanceOwner;
+
+	friend class FFastGeoPrimitiveComponent;
+	friend class FFastGeoInstancedStaticMeshComponent;
+};
\ No newline at end of file
diff --git a/Engine/Plugins/Experimental/FastGeoStreaming/Source/FastGeoStreaming/Internal/FastGeoPrimitiveComponent.h b/Engine/Plugins/Experimental/FastGeoStreaming/Source/FastGeoStreaming/Internal/FastGeoPrimitiveComponent.h
--- a/Engine/Plugins/Experimental/FastGeoStreaming/Source/FastGeoStreaming/Internal/FastGeoPrimitiveComponent.h
+++ b/Engine/Plugins/Experimental/FastGeoStreaming/Source/FastGeoStreaming/Internal/FastGeoPrimitiveComponent.h
@@ -3,14 +3,13 @@
 #pragma once
 
 #include "FastGeoComponent.h"
+#include "FastGeoPhysicsBodyInstanceOwner.h"
 #include "Components/PrimitiveComponent.h"
 #include "Engine/EngineTypes.h"
 #include "PrimitiveSceneInfoData.h"
 #include "PrimitiveSceneProxy.h"
 #include "RenderCommandFence.h"
 #include "PhysicsEngine/BodyInstance.h"
-#include "PhysicsEngine/PhysicsBodyInstanceOwnerInterface.h"
-#include "Chaos/ChaosUserEntity.h"
 #include "Templates/DontCopy.h"
 
 class FRegisterComponentContext;
@@ -25,39 +24,6 @@ struct FPrimitiveSceneProxyDesc;
 struct FPSOPrecacheParams;
 enum class EPSOPrecachePriority : uint8;
 
-class FFastGeoPhysicsBodyInstanceOwner : public FChaosUserDefinedEntity, public IPhysicsBodyInstanceOwner
-{
-public:
-	FFastGeoPhysicsBodyInstanceOwner();
-	virtual ~FFastGeoPhysicsBodyInstanceOwner() = default;
-
-	/** Returns the IPhysicsBodyInstanceOwner based on the provided FChaosUserDefinedEntity */
-	static IPhysicsBodyInstanceOwner* GetPhysicsBodyInstanceOwner(FChaosUserDefinedEntity* InUserDefinedEntity);
-
-	//~ Begin FChaosUserDefinedEntity interface
-	TWeakObjectPtr<UObject> GetOwnerObject() override;
-	//~ End FChaosUserDefinedEntity interface
-
-	//~ Begin IPhysicsBodyInstanceOwner interface
-	virtual bool IsStaticPhysics() const override;
-	virtual UObject* GetSourceObject() const override;
-	virtual ECollisionResponse GetCollisionResponseToChannel(ECollisionChannel Channel) const override;
-	virtual UPhysicalMaterial* GetPhysicalMaterial() const override;
-	virtual void GetComplexPhysicalMaterials(TArray<UPhysicalMaterial*>& OutPhysMaterials, TArray<FPhysicalMaterialMaskParams>* OutPhysMaterialMasks) const override;
-	//~ End IPhysicsBodyInstanceOwner interface
-
-private:
-	void Uninitialize();
-	void Initialize(FFastGeoPrimitiveComponent* InOwner);
-
-	FFastGeoPrimitiveComponent* OwnerComponent;
-	TWeakObjectPtr<UFastGeoContainer> OwnerContainer;
-	static const FName NAME_FastGeoPhysicsBodyInstanceOwner;
-
-	friend class FFastGeoPrimitiveComponent;
-	friend class FFastGeoInstancedStaticMeshComponent;
-};
-
 class FASTGEOSTREAMING_API FFastGeoPrimitiveComponent : public FFastGeoComponent, public IPrimitiveComponent
 {
 public:
@@ -184,10 +150,14 @@ protected:
 	//~ Begin Physics
 	virtual bool IsStaticPhysics() const;
 	virtual UObject* GetSourceObject() const;
+	virtual ECollisionChannel GetCollisionObjectType() const;
+	virtual ECollisionEnabled::Type GetCollisionEnabled() const;
 	virtual UPhysicalMaterial* GetPhysicalMaterial() const { return nullptr; }
 	virtual void GetComplexPhysicalMaterials(TArray<UPhysicalMaterial*>& OutPhysMaterials, TArray<FPhysicalMaterialMaskParams>* OutPhysMaterialMasks) const {}
 	virtual ECollisionResponse GetCollisionResponseToChannel(ECollisionChannel Channel) const;
-	//~ End Navigation
+	virtual Chaos::FPhysicsObject* GetPhysicsObjectById(Chaos::FPhysicsObjectId Id) const;
+	virtual TArray<Chaos::FPhysicsObject*> GetAllPhysicsObjects() const;
+	//~ End Physics
 
 protected:
 	// Persistent data
diff --git a/Engine/Plugins/Experimental/FastGeoStreaming/Source/FastGeoStreaming/Private/FastGeoInstancedStaticMeshComponent.cpp b/Engine/Plugins/Experimental/FastGeoStreaming/Source/FastGeoStreaming/Private/FastGeoInstancedStaticMeshComponent.cpp
--- a/Engine/Plugins/Experimental/FastGeoStreaming/Source/FastGeoStreaming/Private/FastGeoInstancedStaticMeshComponent.cpp
+++ b/Engine/Plugins/Experimental/FastGeoStreaming/Source/FastGeoStreaming/Private/FastGeoInstancedStaticMeshComponent.cpp
@@ -234,6 +234,20 @@ void FFastGeoInstancedStaticMeshComponent::OnAsyncDestroyPhysicsState()
 	AsyncDestroyPhysicsStatePayload.Empty();
 }
 
+Chaos::FPhysicsObject* FFastGeoInstancedStaticMeshComponent::GetPhysicsObjectById(Chaos::FPhysicsObjectId Id) const
+{
+	if (!InstanceBodies.IsValidIndex(Id) || !InstanceBodies[Id] || !InstanceBodies[Id]->GetPhysicsActor())
+	{
+		return nullptr;
+	}
+	return InstanceBodies[Id]->GetPhysicsActor()->GetPhysicsObject();
+}
+
+TArray<Chaos::FPhysicsObject*> FFastGeoInstancedStaticMeshComponent::GetAllPhysicsObjects() const
+{
+	return GetAllPhysicsObjects(InstanceBodies);
+}
+
 TArray<Chaos::FPhysicsObject*> FFastGeoInstancedStaticMeshComponent::GetAllPhysicsObjects(TArray<FBodyInstance*> InInstanceBodies)
 {
 	TArray<Chaos::FPhysicsObject*> Objects;
diff --git a/Engine/Plugins/Experimental/FastGeoStreaming/Source/FastGeoStreaming/Private/FastGeoPhysicsBodyInstanceOwner.cpp b/Engine/Plugins/Experimental/FastGeoStreaming/Source/FastGeoStreaming/Private/FastGeoPhysicsBodyInstanceOwner.cpp
new file mode 100644
--- /dev/null
+++ b/Engine/Plugins/Experimental/FastGeoStreaming/Source/FastGeoStreaming/Private/FastGeoPhysicsBodyInstanceOwner.cpp
@@ -0,0 +1,103 @@
+// Copyright Epic Games, Inc. All Rights Reserved.
+
+#include "FastGeoPhysicsBodyInstanceOwner.h"
+#include "FastGeoPrimitiveComponent.h"
+#include "FastGeoContainer.h"
+
+const FName FFastGeoPhysicsBodyInstanceOwner::NAME_FastGeoPhysicsBodyInstanceOwner = TEXT("FastGeoPhysicsBodyInstanceOwner");
+
+FFastGeoPhysicsBodyInstanceOwner::FFastGeoPhysicsBodyInstanceOwner()
+	: FChaosUserDefinedEntity(FFastGeoPhysicsBodyInstanceOwner::NAME_FastGeoPhysicsBodyInstanceOwner)
+	, OwnerComponent(nullptr)
+{
+}
+
+void FFastGeoPhysicsBodyInstanceOwner::Uninitialize()
+{
+	Initialize(nullptr);
+}
+
+void FFastGeoPhysicsBodyInstanceOwner::Initialize(FFastGeoPrimitiveComponent* InOwnerComponent)
+{
+	OwnerComponent = InOwnerComponent;
+	UFastGeoContainer* NewContainer = OwnerComponent ? OwnerComponent->GetOwnerContainer() : nullptr;
+	check(!OwnerComponent || NewContainer);
+	check(OwnerContainer.IsExplicitlyNull() || !NewContainer || OwnerContainer == NewContainer);
+	OwnerContainer = NewContainer;
+}
+
+TWeakObjectPtr<UObject> FFastGeoPhysicsBodyInstanceOwner::GetOwnerObject()
+{
+	return OwnerContainer;
+}
+
+IPhysicsBodyInstanceOwner* FFastGeoPhysicsBodyInstanceOwner::GetPhysicsBodyInstanceOwner(FChaosUserDefinedEntity* InUserDefinedEntity)
+{
+	if (InUserDefinedEntity && InUserDefinedEntity->GetEntityTypeName() == NAME_FastGeoPhysicsBodyInstanceOwner)
+	{
+		FFastGeoPhysicsBodyInstanceOwner* FastGeoPhysicsBodyInstanceOwner = static_cast<FFastGeoPhysicsBodyInstanceOwner*>(InUserDefinedEntity);
+		check(FastGeoPhysicsBodyInstanceOwner->GetOwnerObject().IsValid());
+		return FastGeoPhysicsBodyInstanceOwner;
+	}
+	return nullptr;
+}
+
+bool FFastGeoPhysicsBodyInstanceOwner::IsStaticPhysics() const
+{
+	check(OwnerContainer.IsValid());
+	return OwnerComponent->IsStaticPhysics();
+}
+
+UObject* FFastGeoPhysicsBodyInstanceOwner::GetSourceObject() const
+{
+	check(OwnerContainer.IsValid());
+	return OwnerComponent->GetSourceObject();
+}
+
+FTransform FFastGeoPhysicsBodyInstanceOwner::GetTransform() const
+{
+	check(OwnerContainer.IsValid());
+	return OwnerComponent->GetComponentTransform();
+}
+
+ECollisionChannel FFastGeoPhysicsBodyInstanceOwner::GetCollisionObjectType() const
+{
+	check(OwnerContainer.IsValid());
+	return OwnerComponent->GetCollisionObjectType();
+}
+
+ECollisionEnabled::Type FFastGeoPhysicsBodyInstanceOwner::GetCollisionEnabled() const
+{
+	check(OwnerContainer.IsValid());
+	return OwnerComponent->GetCollisionEnabled();
+}
+
+ECollisionResponse FFastGeoPhysicsBodyInstanceOwner::GetCollisionResponseToChannel(ECollisionChannel Channel) const
+{
+	check(OwnerContainer.IsValid());
+	return OwnerComponent->GetCollisionResponseToChannel(Channel);
+}
+
+UPhysicalMaterial* FFastGeoPhysicsBodyInstanceOwner::GetPhysicalMaterial() const
+{
+	check(OwnerContainer.IsValid());
+	return OwnerComponent->GetPhysicalMaterial();
+}
+
+void FFastGeoPhysicsBodyInstanceOwner::GetComplexPhysicalMaterials(TArray<UPhysicalMaterial*>& OutPhysMaterials, TArray<FPhysicalMaterialMaskParams>* OutPhysMaterialMasks) const
+{
+	check(OwnerContainer.IsValid());
+	return OwnerComponent->GetComplexPhysicalMaterials(OutPhysMaterials, OutPhysMaterialMasks);
+}
+
+Chaos::FPhysicsObject* FFastGeoPhysicsBodyInstanceOwner::GetPhysicsObjectById(Chaos::FPhysicsObjectId Id) const
+{
+	check(OwnerContainer.IsValid());
+	return OwnerComponent->GetPhysicsObjectById(Id);
+}
+
+TArray<Chaos::FPhysicsObject*> FFastGeoPhysicsBodyInstanceOwner::GetAllPhysicsObjects() const
+{
+	check(OwnerContainer.IsValid());
+	return OwnerComponent->GetAllPhysicsObjects();
+}
\ No newline at end of file
diff --git a/Engine/Plugins/Experimental/FastGeoStreaming/Source/FastGeoStreaming/Private/FastGeoPrimitiveComponent.cpp b/Engine/Plugins/Experimental/FastGeoStreaming/Source/FastGeoStreaming/Private/FastGeoPrimitiveComponent.cpp
--- a/Engine/Plugins/Experimental/FastGeoStreaming/Source/FastGeoStreaming/Private/FastGeoPrimitiveComponent.cpp
+++ b/Engine/Plugins/Experimental/FastGeoStreaming/Source/FastGeoStreaming/Private/FastGeoPrimitiveComponent.cpp
@@ -839,69 +839,37 @@ FPrimitiveMaterialPropertyDescriptor FFastGeoPrimitiveComponent::GetUsedMaterial
 	return FPrimitiveComponentHelper::GetUsedMaterialPropertyDesc(*this, InShaderPlatform);
 }
 
-const FName FFastGeoPhysicsBodyInstanceOwner::NAME_FastGeoPhysicsBodyInstanceOwner = TEXT("FastGeoPhysicsBodyInstanceOwner");
-
-FFastGeoPhysicsBodyInstanceOwner::FFastGeoPhysicsBodyInstanceOwner()
-	: FChaosUserDefinedEntity(FFastGeoPhysicsBodyInstanceOwner::NAME_FastGeoPhysicsBodyInstanceOwner)
-	, OwnerComponent(nullptr)
-{}
-
-void FFastGeoPhysicsBodyInstanceOwner::Uninitialize()
-{
-	Initialize(nullptr);
-}
-
-void FFastGeoPhysicsBodyInstanceOwner::Initialize(FFastGeoPrimitiveComponent* InOwnerComponent)
+ECollisionChannel FFastGeoPrimitiveComponent::GetCollisionObjectType() const
 {
-	OwnerComponent = InOwnerComponent;
-	UFastGeoContainer* NewContainer = OwnerComponent ? OwnerComponent->GetOwnerContainer() : nullptr;
-	check(!OwnerComponent || NewContainer);
-	check(OwnerContainer.IsExplicitlyNull() || !NewContainer || OwnerContainer == NewContainer);
-	OwnerContainer = NewContainer;
+	return ECollisionChannel(BodyInstance.GetObjectType());
 }
 
-TWeakObjectPtr<UObject> FFastGeoPhysicsBodyInstanceOwner::GetOwnerObject()
+ECollisionEnabled::Type FFastGeoPrimitiveComponent::GetCollisionEnabled() const
 {
-	return OwnerContainer;
-}
-
-IPhysicsBodyInstanceOwner* FFastGeoPhysicsBodyInstanceOwner::GetPhysicsBodyInstanceOwner(FChaosUserDefinedEntity* InUserDefinedEntity)
-{
-	if (InUserDefinedEntity && InUserDefinedEntity->GetEntityTypeName() == NAME_FastGeoPhysicsBodyInstanceOwner)
+	if (!IsCollisionEnabled())
 	{
-		FFastGeoPhysicsBodyInstanceOwner* FastGeoPhysicsBodyInstanceOwner = static_cast<FFastGeoPhysicsBodyInstanceOwner*>(InUserDefinedEntity);
-		check(FastGeoPhysicsBodyInstanceOwner->GetOwnerObject().IsValid());
-		return FastGeoPhysicsBodyInstanceOwner;
+		return ECollisionEnabled::NoCollision;
 	}
-	return nullptr;
-}
-
-bool FFastGeoPhysicsBodyInstanceOwner::IsStaticPhysics() const
-{
-	check(OwnerContainer.IsValid());
-	return OwnerComponent->IsStaticPhysics();
-}
 
-UObject* FFastGeoPhysicsBodyInstanceOwner::GetSourceObject() const
-{
-	check(OwnerContainer.IsValid());
-	return OwnerComponent->GetSourceObject();
+	return BodyInstance.GetCollisionEnabled(false);
 }
 
-ECollisionResponse FFastGeoPhysicsBodyInstanceOwner::GetCollisionResponseToChannel(ECollisionChannel Channel) const 
+Chaos::FPhysicsObject* FFastGeoPrimitiveComponent::GetPhysicsObjectById(Chaos::FPhysicsObjectId Id) const
 {
-	check(OwnerContainer.IsValid());
-	return OwnerComponent->GetCollisionResponseToChannel(Channel);
-}
+	if (!BodyInstance.IsValidBodyInstance())
+	{
+		return nullptr;
+	}
 
-UPhysicalMaterial* FFastGeoPhysicsBodyInstanceOwner::GetPhysicalMaterial() const
-{
-	check(OwnerContainer.IsValid());
-	return OwnerComponent->GetPhysicalMaterial();
+	return BodyInstance.GetPhysicsActor()->GetPhysicsObject();
 }
 
-void FFastGeoPhysicsBodyInstanceOwner::GetComplexPhysicalMaterials(TArray<UPhysicalMaterial*>& OutPhysMaterials, TArray<FPhysicalMaterialMaskParams>* OutPhysMaterialMasks) const
+TArray<Chaos::FPhysicsObject*> FFastGeoPrimitiveComponent::GetAllPhysicsObjects() const
 {
-	check(OwnerContainer.IsValid());
-	return OwnerComponent->GetComplexPhysicalMaterials(OutPhysMaterials, OutPhysMaterialMasks);
-}
+	TArray<Chaos::FPhysicsObject*> Bodies;
+	if (Chaos::FPhysicsObject* PhysicsObject = GetPhysicsObjectById(INDEX_NONE))
+	{
+		Bodies.Add(PhysicsObject);
+	}
+	return Bodies;
+}
\ No newline at end of file
diff --git a/Engine/Source/Runtime/Engine/Classes/PhysicsEngine/PhysicsBodyInstanceOwnerInterface.h b/Engine/Source/Runtime/Engine/Classes/PhysicsEngine/PhysicsBodyInstanceOwnerInterface.h
--- a/Engine/Source/Runtime/Engine/Classes/PhysicsEngine/PhysicsBodyInstanceOwnerInterface.h
+++ b/Engine/Source/Runtime/Engine/Classes/PhysicsEngine/PhysicsBodyInstanceOwnerInterface.h
@@ -5,6 +5,7 @@
 #include "UObject/Interface.h"
 #include "Containers/Array.h"
 #include "Chaos/PhysicsObject.h"
+#include "Engine/EngineTypes.h"
 #include "PhysicsBodyInstanceOwnerInterface.generated.h"
 
 class UObject;
@@ -60,12 +61,27 @@ public:
 	/** Source object for this body. */
 	virtual UObject* GetSourceObject() const = 0;
 
+	/** Returns the transform of this body instance owner. */
+	virtual FTransform GetTransform() const = 0;
+
+	/** Returns the channel that the body instance owner belongs to. */
+	virtual ECollisionChannel GetCollisionObjectType() const = 0;
+
+	/** Returns the form of collision for this body instance owner. */
+	virtual ECollisionEnabled::Type GetCollisionEnabled() const = 0;
+
 	/** Find the correct PhysicalMaterial for simple geometry on this body (used by FBodyInstance::GetSimplePhysicalMaterial). */
 	virtual UPhysicalMaterial* GetPhysicalMaterial() const = 0;
 
-	/** Get the complex PhysicalMaterials array for this body (used by FBodyInstance::GetComplexPhysicalMaterials). */
+	/** Gets the complex PhysicalMaterials array for this body (used by FBodyInstance::GetComplexPhysicalMaterials). */
 	virtual void GetComplexPhysicalMaterials(TArray<UPhysicalMaterial*>& OutPhysMaterials, TArray<FPhysicalMaterialMaskParams>* OutPhysMaterialMasks) const = 0;
 
 	/** Gets the response type given a specific channel. */
 	virtual ECollisionResponse GetCollisionResponseToChannel(ECollisionChannel Channel) const = 0;
+
+	/** Returns a physics object corresponding to the specified Id. */
+	virtual Chaos::FPhysicsObject* GetPhysicsObjectById(Chaos::FPhysicsObjectId Id) const = 0;
+
+	/** Returns all physics objects associated to the physics body instance owner. */
+	virtual TArray<Chaos::FPhysicsObject*> GetAllPhysicsObjects() const = 0;
 };
\ No newline at end of file
