[FastGeo] Fixed a crash occurring when switching the Preview Platform during PIE with FastGeo content.
The fix generalizes FastGeo?s handling of global component re-registration by reacting to FGlobalComponentReregisterContext lifetime events.
Notifications are triggered when the first context is created and when the final active context is destroyed.

#Synced-CL 49593614

diff --git a/Engine/Plugins/Experimental/FastGeoStreaming/Source/FastGeoStreaming/Private/FastGeoWorldSubsystem.cpp b/Engine/Plugins/Experimental/FastGeoStreaming/Source/FastGeoStreaming/Private/FastGeoWorldSubsystem.cpp
--- a/Engine/Plugins/Experimental/FastGeoStreaming/Source/FastGeoStreaming/Private/FastGeoWorldSubsystem.cpp
+++ b/Engine/Plugins/Experimental/FastGeoStreaming/Source/FastGeoStreaming/Private/FastGeoWorldSubsystem.cpp
@@ -6,6 +6,7 @@
 #include "FastGeoStreamingModule.h"
 #include "Physics/Experimental/PhysScene_Chaos.h"
 #include "Streaming/LevelStreamingDelegates.h"
+#include "ComponentReregisterContext.h"
 #include "Engine/Engine.h"
 #include "Engine/World.h"
 #include "Engine/Level.h"
@@ -173,6 +174,11 @@ void UFastGeoWorldSubsystem::PostInitialize()
 	UWorld* World = GetWorld();
 
 	FWorldDelegates::OnWorldCleanup.AddUObject(this, &UFastGeoWorldSubsystem::OnWorldCleanup);
+	FGlobalComponentReregisterContext::OnGlobalComponentReregisterContextCreated().AddUObject(this, &UFastGeoWorldSubsystem::OnGlobalComponentReregisterContextCreated);
+	FGlobalComponentReregisterContext::OnGlobalComponentReregisterContextDestroyed().AddUObject(this, &UFastGeoWorldSubsystem::OnGlobalComponentReregisterContextDestroyed);
+#if WITH_EDITOR
+	FWorldDelegates::OnPreRecreateScene.AddUObject(this, &UFastGeoWorldSubsystem::OnPreRecreateScene);
+#endif
 
 	if (World->IsPartitionedWorld() && World->IsGameWorld())
 	{
@@ -198,6 +204,11 @@ void UFastGeoWorldSubsystem::Deinitialize()
 	UWorld* World = GetWorld();
 
 	FWorldDelegates::OnWorldCleanup.RemoveAll(this);
+	FGlobalComponentReregisterContext::OnGlobalComponentReregisterContextCreated().RemoveAll(this);
+	FGlobalComponentReregisterContext::OnGlobalComponentReregisterContextDestroyed().RemoveAll(this);
+#if WITH_EDITOR
+	FWorldDelegates::OnPreRecreateScene.RemoveAll(this);
+#endif
 
 	if (World->IsGameWorld())
 	{
@@ -289,6 +300,59 @@ void UFastGeoWorldSubsystem::OnWorldCleanup(UWorld* World, bool bSessionEnded, b
 	PendingDestroyContainers.Reset();
 }
 
+void UFastGeoWorldSubsystem::OnGlobalComponentReregisterContextCreated()
+{
+	check(ReregisteringFastGeoContainers.IsEmpty());
+
+	WaitForAllAsyncJobsCompletion();
+	
+	TGuardValue<bool> GuardValue(bWaitingForCompletion, true);
+	UWorld* World = GetWorld();
+	for (TObjectIterator<UFastGeoContainer> It(RF_ClassDefaultObject, true, EInternalObjectFlags::Garbage); It; ++It)
+	{
+		UFastGeoContainer* FastGeo = *It;
+		if ((FastGeo->GetWorld() == World) && FastGeo->IsRegistered())
+		{
+			FastGeo->Unregister();
+			ReregisteringFastGeoContainers.Add(FastGeo);
+		}
+	}
+
+	WaitForAllAsyncJobsCompletion();
+}
+
+void UFastGeoWorldSubsystem::OnGlobalComponentReregisterContextDestroyed()
+{
+	WaitForAllAsyncJobsCompletion();
+
+	TGuardValue<bool> GuardValue(bWaitingForCompletion, true);
+	for (const TWeakObjectPtr<UFastGeoContainer>& WeakFastGeo : ReregisteringFastGeoContainers)
+	{
+		UFastGeoContainer* FastGeo = WeakFastGeo.Get();
+		if (FastGeo && !FastGeo->IsRegistered())
+		{
+			FastGeo->Register();
+		}
+	}
+
+	WaitForAllAsyncJobsCompletion();
+
+	ReregisteringFastGeoContainers.Empty();
+}
+
+#if WITH_EDITOR
+void UFastGeoWorldSubsystem::OnPreRecreateScene(UWorld* World)
+{
+	if (World == GetWorld())
+	{
+		// Make sure to flush AsyncRenderStateJobQueue, as it holds a raw pointer to the scene until it completes
+		const bool bWaitForCompletion = true;
+		ProcessAsyncRenderStateJobs(bWaitForCompletion);
+		check(!AsyncRenderStateJobQueue.IsValid());
+	}
+}
+#endif // WITH_EDITOR
+
 void UFastGeoWorldSubsystem::OnLevelComponentsUpdated(UWorld* World, ULevel* Level)
 {
 	if (World != GetWorld())
@@ -375,17 +439,20 @@ void UFastGeoWorldSubsystem::OnRemoveLevelFromWorldExtension(ULevel* InLevel, co
 
 void UFastGeoWorldSubsystem::OnUpdateLevelStreaming()
 {
-	UWorld* World = GetWorld();
-	if (World->GetShouldForceUnloadStreamingLevels())
+	if (GetWorld()->GetShouldForceUnloadStreamingLevels())
 	{
 		TRACE_CPUPROFILER_EVENT_SCOPE(UFastGeoWorldSubsystem::OnUpdateLevelStreaming);
+		WaitForAllAsyncJobsCompletion();
+	}
+}
 
-		const bool bWaitForCompletion = true;
-		ProcessAsyncRenderStateJobs(bWaitForCompletion);
-		if (FPhysScene* PhysScene = World->GetPhysicsScene())
-		{
-			PhysScene->ProcessAsyncPhysicsStateJobs(bWaitForCompletion);
-		}
+void UFastGeoWorldSubsystem::WaitForAllAsyncJobsCompletion()
+{
+	const bool bWaitForCompletion = true;
+	ProcessAsyncRenderStateJobs(bWaitForCompletion);
+	if (FPhysScene* PhysScene = GetWorld()->GetPhysicsScene())
+	{
+		PhysScene->ProcessAsyncPhysicsStateJobs(bWaitForCompletion);
 	}
 }
 
diff --git a/Engine/Plugins/Experimental/FastGeoStreaming/Source/FastGeoStreaming/Private/FastGeoWorldSubsystem.h b/Engine/Plugins/Experimental/FastGeoStreaming/Source/FastGeoStreaming/Private/FastGeoWorldSubsystem.h
--- a/Engine/Plugins/Experimental/FastGeoStreaming/Source/FastGeoStreaming/Private/FastGeoWorldSubsystem.h
+++ b/Engine/Plugins/Experimental/FastGeoStreaming/Source/FastGeoStreaming/Private/FastGeoWorldSubsystem.h
@@ -95,20 +95,22 @@ private:
 	void OnLevelComponentsCleared(UWorld* World, ULevel* Level);
 	void OnAddLevelToWorldExtension(ULevel* Level, const bool bWaitForCompletion, bool& bOutHasCompleted);
 	void OnRemoveLevelFromWorldExtension(ULevel* Level, const bool bWaitForCompletion, bool& bOutHasCompleted);
+	void OnGlobalComponentReregisterContextCreated();
+	void OnGlobalComponentReregisterContextDestroyed();
 #if DO_CHECK
 	void OnLevelAddedToWorld(ULevel* Level, UWorld* World);
 	void OnLevelRemovedFromWorld(ULevel* Level, UWorld* World);
 	void CheckNoPendingTasks(ULevel* Level, UWorld* World);
 #endif
-
+#if WITH_EDITOR
+	void OnPreRecreateScene(UWorld* World);
+#endif
 	void ForEachHLODObjectInCell(const UWorldPartitionRuntimeCell*, TFunction<void(IWorldPartitionHLODObject*)>);
-
 	void RequestAsyncRenderStateTasksBudget_Concurrent(float& OutAvailableTimeBudgetMS, int32& OutAvaiableComponentsBudget, int32& OutTimeEpoch);
 	void CommitAsyncRenderStateTasksBudget_Concurrent(float InUsedTimeBudgetMS, int32& InUsedComponentsBudget, int32 TimeEpoch);
-
 	void PushAsyncDestroyContainer(UFastGeoContainer* FastGeo);
-
 	FFastGeoAsyncRenderStateJobQueue* GetOrCreateJobQueue();
+	void WaitForAllAsyncJobsCompletion();
 
 	UPROPERTY()
 	TSet<TObjectPtr<UFastGeoContainer>> PendingDestroyContainers;
@@ -129,6 +131,7 @@ private:
 	bool bWaitingForCompletion = false;
 
 	TUniquePtr<FFastGeoAsyncRenderStateJobQueue> AsyncRenderStateJobQueue;
+	TArray<TWeakObjectPtr<UFastGeoContainer>> ReregisteringFastGeoContainers;
 
 	friend class UFastGeoContainer;
 };
diff --git a/Engine/Source/Runtime/Engine/Private/Components/ActorComponent.cpp b/Engine/Source/Runtime/Engine/Private/Components/ActorComponent.cpp
--- a/Engine/Source/Runtime/Engine/Private/Components/ActorComponent.cpp
+++ b/Engine/Source/Runtime/Engine/Private/Components/ActorComponent.cpp
@@ -98,6 +98,12 @@ FAutoConsoleVariableRef GTickComponentLatentActionsWithTheComponentCVar(
 /** Enable to log out all render state create, destroy and updatetransform events */
 #define LOG_RENDER_STATE 0
 
+/** Static delegate broadcasted when the first active FGlobalComponentReregisterContext is created. */
+FGlobalComponentReregisterContext::FOnGlobalComponentReregisterContextCreated FGlobalComponentReregisterContext::GlobalComponentReregisterContextCreatedEvent;
+
+/** Static delegate broadcasted when the last active FGlobalComponentReregisterContext is destroyed. */
+FGlobalComponentReregisterContext::FOnGlobalComponentReregisterContextDestroyed FGlobalComponentReregisterContext::GlobalComponentReregisterContextDestroyedEvent;
+
 /** Static var indicating activity of reregister context */
 int32 FGlobalComponentReregisterContext::ActiveGlobalReregisterContextCount = 0;
 
@@ -365,6 +371,12 @@ FGlobalComponentReregisterContext::FGlobalComponentReregisterContext()
 		ComponentContexts.Add(new FComponentReregisterContext(Component, &ScenesToUpdateAllPrimitiveSceneInfos));
 	}
 
+	// Only broadcast if this is the first active context
+	if (ActiveGlobalReregisterContextCount == 1)
+	{
+		GlobalComponentReregisterContextCreatedEvent.Broadcast();
+	}
+
 	UpdateAllPrimitiveSceneInfos();
 }
 
@@ -394,6 +406,8 @@ FGlobalComponentReregisterContext::FGlobalComponentReregisterContext(const TArra
 				ComponentContexts.Add(new FComponentReregisterContext(Component, &ScenesToUpdateAllPrimitiveSceneInfos));
 			}
 		}
+
+		GlobalComponentReregisterContextCreatedEvent.Broadcast();
 		
 		UpdateAllPrimitiveSceneInfos();
 	}
@@ -407,6 +421,8 @@ FGlobalComponentReregisterContext::~FGlobalComponentReregisterContext()
 	if (--ActiveGlobalReregisterContextCount == 0)
 	{
 		ComponentContexts.Empty();
+
+		GlobalComponentReregisterContextDestroyedEvent.Broadcast();
 		
 		UpdateAllPrimitiveSceneInfos();
 	}
diff --git a/Engine/Source/Runtime/Engine/Public/ComponentReregisterContext.h b/Engine/Source/Runtime/Engine/Public/ComponentReregisterContext.h
--- a/Engine/Source/Runtime/Engine/Public/ComponentReregisterContext.h
+++ b/Engine/Source/Runtime/Engine/Public/ComponentReregisterContext.h
@@ -147,12 +147,24 @@ public:
 	/** Indicates that a FGlobalComponentReregisterContext is currently active */
 	static int32 ActiveGlobalReregisterContextCount;
 
+	DECLARE_EVENT(FGlobalComponentReregisterContext, FOnGlobalComponentReregisterContextCreated);
+	static FOnGlobalComponentReregisterContextCreated& OnGlobalComponentReregisterContextCreated() { return GlobalComponentReregisterContextCreatedEvent; }
+
+	DECLARE_EVENT(FGlobalComponentReregisterContext, FOnGlobalComponentReregisterContextDestroyed);
+	static FOnGlobalComponentReregisterContextDestroyed& OnGlobalComponentReregisterContextDestroyed() { return GlobalComponentReregisterContextDestroyedEvent; }
+
 private:
 	/** The recreate contexts for the individual components. */
 	TIndirectArray<FComponentReregisterContext> ComponentContexts;
 
 	TSet<FSceneInterface*> ScenesToUpdateAllPrimitiveSceneInfos;
 
+	/** Static delegate broadcasted when the first active FGlobalComponentReregisterContext is created. */
+	ENGINE_API static FOnGlobalComponentReregisterContextCreated GlobalComponentReregisterContextCreatedEvent;
+
+	/** Static delegate broadcasted when the last active FGlobalComponentReregisterContext is destroyed. */
+	ENGINE_API static FOnGlobalComponentReregisterContextDestroyed GlobalComponentReregisterContextDestroyedEvent;
+
 	void UpdateAllPrimitiveSceneInfos();
 };
 
