[FastGeo] Fix compile error when UE_WITH_PSO_PRECACHING is not set.

#rnx
#Synced-CL 49721941

diff --git a/Engine/Plugins/Experimental/FastGeoStreaming/Source/FastGeoStreaming/Internal/FastGeoPSOWeakPtrTraits.h b/Engine/Plugins/Experimental/FastGeoStreaming/Source/FastGeoStreaming/Internal/FastGeoPSOWeakPtrTraits.h
--- a/Engine/Plugins/Experimental/FastGeoStreaming/Source/FastGeoStreaming/Internal/FastGeoPSOWeakPtrTraits.h
+++ b/Engine/Plugins/Experimental/FastGeoStreaming/Source/FastGeoStreaming/Internal/FastGeoPSOWeakPtrTraits.h
@@ -5,6 +5,8 @@
 #include "PSOPrecacheComponent.h"
 #include "FastGeoWeakElement.h"
 
+#if UE_WITH_PSO_PRECACHING
+
 // Specialization: if T : FFastGeoComponent -> use FWeakFastGeoComponent
 template<typename T>
 struct TPSOWeakPtrTraits<T, typename TEnableIf<TIsDerivedFrom<T, FFastGeoComponent>::Value>::Type>
@@ -13,4 +15,6 @@ struct TPSOWeakPtrTraits<T, typename TEnableIf<TIsDerivedFrom<T, FFastGeoCompone
 
 	static WeakType Make(T* Ptr) { return WeakType(Ptr); }
 	static T* Get(const WeakType& W) { return W.Get<T>(); }
-};
\ No newline at end of file
+};
+
+#endif // UE_WITH_PSO_PRECACHING
