[FastGeo] Added option to store FastGeo ISM instance transforms either in FTransform3f or in FMatrix based on a transformer setting EFastGeoInstanceStorageMode (Auto, Compressed, Full).

#Synced-CL 49728154

diff --git a/Engine/Plugins/Experimental/FastGeoStreaming/Source/FastGeoStreaming/Internal/FastGeoInstancedStaticMeshComponent.h b/Engine/Plugins/Experimental/FastGeoStreaming/Source/FastGeoStreaming/Internal/FastGeoInstancedStaticMeshComponent.h
--- a/Engine/Plugins/Experimental/FastGeoStreaming/Source/FastGeoStreaming/Internal/FastGeoInstancedStaticMeshComponent.h
+++ b/Engine/Plugins/Experimental/FastGeoStreaming/Source/FastGeoStreaming/Internal/FastGeoInstancedStaticMeshComponent.h
@@ -37,6 +37,9 @@ public:
 	virtual TArray<Chaos::FPhysicsObject*> GetAllPhysicsObjects() const override;
 	//~ End FFastGeoPrimitiveComponent interface
 
+	void ForEachInstanceMatrix(TFunctionRef<void(const FMatrix&)> Func) const;
+	int32 GetInstanceCount() const;
+
 protected:
 	//~ Begin FFastGeoPrimitiveComponent interface
 	virtual FPrimitiveSceneProxyDesc& GetSceneProxyDesc() override { return SceneProxyDesc; }
@@ -69,10 +72,15 @@ private:
 	FBoxSphereBounds CalculateBounds(EBoundsType BoundsType);
 	void CreateAllInstanceBodies();
 	static TArray<Chaos::FPhysicsObject*> GetAllPhysicsObjects(TArray<FBodyInstance*> InInstanceBodies);
+#if WITH_EDITOR
+	void InitializePerInstanceSMData(const TArray<FInstancedStaticMeshInstanceData>& InPerInstanceSMData);
+#endif
 
 private:
 	// Persistent Data
-	TArray<FInstancedStaticMeshInstanceData> PerInstanceSMData;
+	TArray<FInstancedStaticMeshInstanceData> HighPrecisionPerInstanceSMData;
+	TArray<FTransform3f> LowPrecisionPerInstanceSMData;
+	bool bUseHighPrecisionPerInstanceSMData = false;
 	int32 InstancingRandomSeed = 0;
 	TArray<float> PerInstanceSMCustomData;
 	TArray<FInstancedStaticMeshRandomSeed> AdditionalRandomSeeds;
diff --git a/Engine/Plugins/Experimental/FastGeoStreaming/Source/FastGeoStreaming/Private/FastGeoInstancedStaticMeshComponent.cpp b/Engine/Plugins/Experimental/FastGeoStreaming/Source/FastGeoStreaming/Private/FastGeoInstancedStaticMeshComponent.cpp
--- a/Engine/Plugins/Experimental/FastGeoStreaming/Source/FastGeoStreaming/Private/FastGeoInstancedStaticMeshComponent.cpp
+++ b/Engine/Plugins/Experimental/FastGeoStreaming/Source/FastGeoStreaming/Private/FastGeoInstancedStaticMeshComponent.cpp
@@ -1,6 +1,7 @@
 // Copyright Epic Games, Inc. All Rights Reserved.
 
 #include "FastGeoInstancedStaticMeshComponent.h"
+#include "FastGeoWorldPartitionRuntimeCellTransformer.h"
 #include "FastGeoComponentCluster.h"
 #include "FastGeoContainer.h"
 #include "FastGeoHLOD.h"
@@ -29,7 +30,15 @@ void FFastGeoInstancedStaticMeshComponent::Serialize(FArchive& Ar)
 	Super::Serialize(Ar);
 
 	// Serialize persistent data from FFastGeoInstancedStaticMeshComponent
-	PerInstanceSMData.BulkSerialize(Ar);
+	Ar << bUseHighPrecisionPerInstanceSMData;
+	if (bUseHighPrecisionPerInstanceSMData)
+	{
+		HighPrecisionPerInstanceSMData.BulkSerialize(Ar);
+	}
+	else
+	{
+		LowPrecisionPerInstanceSMData.BulkSerialize(Ar);
+	}
 	Ar << InstancingRandomSeed;
 	PerInstanceSMCustomData.BulkSerialize(Ar);
 	Ar << AdditionalRandomSeeds;
@@ -57,18 +66,70 @@ void FFastGeoInstancedStaticMeshComponent::InitializeFromComponent(UActorCompone
 	UInstancedStaticMeshComponent* ISMComponent = CastChecked<UInstancedStaticMeshComponent>(Component);
 	SceneProxyDesc.bCollisionEnabled = !!SceneProxyDesc.bCollisionEnabled && !ISMComponent->bDisableCollision;
 	AdditionalRandomSeeds = ISMComponent->AdditionalRandomSeeds;
-	PerInstanceSMData = ISMComponent->PerInstanceSMData;
+	InitializePerInstanceSMData(ISMComponent->PerInstanceSMData);
 	PerInstanceSMCustomData = ISMComponent->PerInstanceSMCustomData;
 	InstancingRandomSeed = ISMComponent->InstancingRandomSeed;
 
 	// ISMC with no instances should never be transformed to FastGeo
-	check(!PerInstanceSMData.IsEmpty());
+	check(GetInstanceCount() > 0);
 
 	LocalBounds = CalculateBounds(EBoundsType::LocalBounds);
 	WorldBounds = CalculateBounds(EBoundsType::WorldBounds);
 	NavigationBounds = CalculateBounds(EBoundsType::NavigationBounds).GetBox();
 }
 
+void FFastGeoInstancedStaticMeshComponent::InitializePerInstanceSMData(const TArray<FInstancedStaticMeshInstanceData>& InPerInstanceSMData)
+{
+	const UFastGeoWorldPartitionRuntimeCellTransformer* Transformer = UFastGeoWorldPartitionRuntimeCellTransformer::GetCurrentTransformer();
+	const UFastGeoTransformerSettings& Settings = Transformer ? Transformer->GetSettings() : *GetDefault<UFastGeoTransformerSettings>();
+	const EFastGeoInstanceStorageMode InstanceStorageMode = Settings.InstanceStorageMode;
+	const float PositionThreshold = Settings.InstanceCompressionSettings.PositionEpsilon;
+	const float RotationThreshold = Settings.InstanceCompressionSettings.RotationEpsilon;
+	const float ScaleThreshold = Settings.InstanceCompressionSettings.ScaleEpsilon;
+	const bool bAutomatic = (InstanceStorageMode == EFastGeoInstanceStorageMode::Auto);
+	bUseHighPrecisionPerInstanceSMData = (InstanceStorageMode == EFastGeoInstanceStorageMode::Full);
+	HighPrecisionPerInstanceSMData.Empty();
+	LowPrecisionPerInstanceSMData.Empty();
+
+	if (!bUseHighPrecisionPerInstanceSMData)
+	{
+		LowPrecisionPerInstanceSMData.Reserve(InPerInstanceSMData.Num());
+		for (const FInstancedStaticMeshInstanceData& InstanceData : InPerInstanceSMData)
+		{
+			const FTransform OrigTransform(InstanceData.Transform);
+			const FTransform3f CompressedTransform(OrigTransform);
+
+			if (bAutomatic)
+			{
+				const FVector OrigLocation = OrigTransform.GetLocation();
+				const FQuat OrigRotation = OrigTransform.GetRotation();
+				const FVector OrigScale = OrigTransform.GetScale3D();
+
+				const FVector NewPos = (FVector)CompressedTransform.GetTranslation();
+				const FQuat NewRot = (FQuat)CompressedTransform.GetRotation();
+				const FVector NewScale = (FVector)CompressedTransform.GetScale3D();
+
+				const bool bPosOk = OrigLocation.Equals(NewPos, PositionThreshold);
+				const bool bRotOk = OrigRotation.Equals(NewRot, RotationThreshold);
+				const bool bScaleOk = OrigScale.Equals(NewScale, ScaleThreshold);
+				if (!bPosOk || !bRotOk || !bScaleOk)
+				{
+					bUseHighPrecisionPerInstanceSMData = true;
+					break;
+				}
+			}
+
+			LowPrecisionPerInstanceSMData.Emplace(CompressedTransform);
+		}
+	}
+
+	if (bUseHighPrecisionPerInstanceSMData)
+	{
+		LowPrecisionPerInstanceSMData.Empty();
+		HighPrecisionPerInstanceSMData = InPerInstanceSMData;
+	}
+}
+
 void FFastGeoInstancedStaticMeshComponent::ResetSceneProxyDescUnsupportedProperties()
 {
 	Super::ResetSceneProxyDescUnsupportedProperties();
@@ -86,6 +147,29 @@ void FFastGeoInstancedStaticMeshComponent::ApplyWorldTransform(const FTransform&
 	NavigationBounds = CalculateBounds(EBoundsType::NavigationBounds).GetBox();
 }
 
+void FFastGeoInstancedStaticMeshComponent::ForEachInstanceMatrix(TFunctionRef<void(const FMatrix&)> Func) const
+{
+	if (bUseHighPrecisionPerInstanceSMData)
+	{
+		for (const FInstancedStaticMeshInstanceData& InstanceData : HighPrecisionPerInstanceSMData)
+		{
+			Func(InstanceData.Transform);
+		}
+	}
+	else
+	{
+		for (const FTransform3f& Transform: LowPrecisionPerInstanceSMData)
+		{
+			Func(FMatrix(Transform.ToMatrixWithScale()));
+		}
+	}
+}
+
+int32 FFastGeoInstancedStaticMeshComponent::GetInstanceCount() const
+{
+	return bUseHighPrecisionPerInstanceSMData ? HighPrecisionPerInstanceSMData.Num() : LowPrecisionPerInstanceSMData.Num();
+}
+
 void FFastGeoInstancedStaticMeshComponent::BuildInstanceData()
 {
 	FInstanceSceneDataBuffers InstanceSceneDataBuffers{};
@@ -106,18 +190,19 @@ void FFastGeoInstancedStaticMeshComponent::BuildInstanceData()
 	View.InstanceLocalBounds.Add(InstanceLocalBounds);
 
 	// LocalToPrimitiveRelativeWorld
-	View.InstanceToPrimitiveRelative.Reserve(PerInstanceSMData.Num());
-	for (const auto& SM : PerInstanceSMData)
+	const int32 InstanceCount = GetInstanceCount();
+	View.InstanceToPrimitiveRelative.Reserve(InstanceCount);
+	ForEachInstanceMatrix([&View](const FMatrix& InstanceTransform)
 	{
-		FRenderTransform InstanceToPrimitive = SM.Transform;
+		FRenderTransform InstanceToPrimitive = InstanceTransform;
 		FRenderTransform LocalToPrimitiveRelativeWorld = InstanceToPrimitive * View.PrimitiveToRelativeWorld;
 		LocalToPrimitiveRelativeWorld.Orthogonalize();
 		View.InstanceToPrimitiveRelative.Add(LocalToPrimitiveRelativeWorld);
-	}
+	});
 
 	// InstanceCustomData
 	View.InstanceCustomData = PerInstanceSMCustomData;
-	View.NumCustomDataFloats = PerInstanceSMCustomData.Num() && PerInstanceSMData.Num() ? PerInstanceSMCustomData.Num() / PerInstanceSMData.Num() : 0;
+	View.NumCustomDataFloats = PerInstanceSMCustomData.Num() && InstanceCount ? PerInstanceSMCustomData.Num() / InstanceCount : 0;
 	View.Flags.bHasPerInstanceCustomData = PrimitiveMaterialDesc.bAnyMaterialHasPerInstanceCustomData && View.NumCustomDataFloats != 0;
 	if (!View.Flags.bHasPerInstanceCustomData)
 	{
@@ -126,10 +211,10 @@ void FFastGeoInstancedStaticMeshComponent::BuildInstanceData()
 	}
 
 	// InstanceRandomIDs
-	View.Flags.bHasPerInstanceRandom = PrimitiveMaterialDesc.bAnyMaterialHasPerInstanceRandom && (PerInstanceSMData.Num() > 0);
+	View.Flags.bHasPerInstanceRandom = PrimitiveMaterialDesc.bAnyMaterialHasPerInstanceRandom && (InstanceCount > 0);
 	if (View.Flags.bHasPerInstanceRandom && InstanceRandomIDs.IsEmpty())
 	{
-		InstanceRandomIDs.SetNumZeroed(PerInstanceSMData.Num());
+		InstanceRandomIDs.SetNumZeroed(InstanceCount);
 		FRandomStream RandomStream = FRandomStream(InstancingRandomSeed);
 		auto AdditionalRandomSeedsIt = AdditionalRandomSeeds.CreateIterator();
 		int32 SeedResetIndex = AdditionalRandomSeedsIt ? AdditionalRandomSeedsIt->StartInstanceIndex : INDEX_NONE;
@@ -157,7 +242,7 @@ FPrimitiveSceneProxy* FFastGeoInstancedStaticMeshComponent::CreateStaticMeshScen
 {
 	check(GetWorld());
 	check(SceneProxyDesc.Scene);
-	check(PerInstanceSMData.Num() > 0);
+	check(GetInstanceCount() > 0);
 
 	BuildInstanceData();
 
@@ -266,7 +351,7 @@ void FFastGeoInstancedStaticMeshComponent::CreateAllInstanceBodies()
 {
 	TRACE_CPUPROFILER_EVENT_SCOPE(FFastGeoInstancedStaticMeshComponent::CreateAllInstanceBodies);
 	
-	const int32 NumBodies = PerInstanceSMData.Num();
+	const int32 NumBodies = GetInstanceCount();
 	check(InstanceBodies.Num() == 0);
 	check(SceneProxyDesc.Mobility != EComponentMobility::Movable);
 
@@ -287,26 +372,27 @@ void FFastGeoInstancedStaticMeshComponent::CreateAllInstanceBodies()
 
 		TArray<FTransform> Transforms;
 		Transforms.Reserve(NumBodies);
-		for (int32 i = 0; i < NumBodies; ++i)
+		int32 InstanceBodyIndex = 0;
+		ForEachInstanceMatrix([this, &InstanceBodyIndex, &Transforms, &InstanceBodiesSanitized](const FMatrix& InstanceTransform)
 		{
-			const FTransform InstanceTM = FTransform(PerInstanceSMData[i].Transform) * WorldTransform;
+			const FTransform InstanceTM = FTransform(InstanceTransform) * WorldTransform;
 			if (InstanceTM.GetScale3D().IsNearlyZero())
 			{
-				InstanceBodies[i] = nullptr;
+				InstanceBodies[InstanceBodyIndex] = nullptr;
 			}
 			else
 			{
 				FBodyInstance* Instance = new FBodyInstance;
-
 				InstanceBodiesSanitized.Add(Instance);
-				InstanceBodies[i] = Instance;
+				InstanceBodies[InstanceBodyIndex] = Instance;
 				Instance->CopyBodyInstancePropertiesFrom(&BodyInstance);
-				Instance->InstanceBodyIndex = i;
+				Instance->InstanceBodyIndex = InstanceBodyIndex;
 				Instance->bAutoWeld = false;
 				Instance->bSimulatePhysics = false;
 				Transforms.Add(InstanceTM);
 			}
-		}
+			++InstanceBodyIndex;
+		});
 
 		if (InstanceBodiesSanitized.Num() > 0)
 		{
@@ -332,7 +418,7 @@ void FFastGeoInstancedStaticMeshComponent::CreateAllInstanceBodies()
 
 FBoxSphereBounds FFastGeoInstancedStaticMeshComponent::CalculateBounds(EBoundsType BoundsType)
 {
-	if (GetStaticMesh() && PerInstanceSMData.Num() > 0)
+	if (GetStaticMesh() && GetInstanceCount() > 0)
 	{
 		const bool bWorldSpace = (BoundsType != EBoundsType::LocalBounds);
 		const FBox InstanceBounds = (BoundsType == EBoundsType::NavigationBounds) ? FInstancedStaticMeshComponentHelper::GetInstanceNavigationBounds(*this) : GetStaticMesh()->GetBounds().GetBox();
@@ -340,17 +426,17 @@ FBoxSphereBounds FFastGeoInstancedStaticMeshComponent::CalculateBounds(EBoundsTy
 		if (InstanceBounds.IsValid)
 		{
 			FBoxSphereBounds::Builder BoundsBuilder;
-			for (int32 InstanceIndex = 0; InstanceIndex < PerInstanceSMData.Num(); InstanceIndex++)
+			ForEachInstanceMatrix([bWorldSpace, &InstanceBounds, &ComponentTransformMatrix, &BoundsBuilder](const FMatrix& InstanceTransform)
 			{
 				if (bWorldSpace)
 				{
-					BoundsBuilder += InstanceBounds.TransformBy(PerInstanceSMData[InstanceIndex].Transform * ComponentTransformMatrix);
+					BoundsBuilder += InstanceBounds.TransformBy(InstanceTransform * ComponentTransformMatrix);
 				}
 				else
 				{
-					BoundsBuilder += InstanceBounds.TransformBy(PerInstanceSMData[InstanceIndex].Transform);
+					BoundsBuilder += InstanceBounds.TransformBy(InstanceTransform);
 				}
-			}
+			});
 			return BoundsBuilder;
 		}
 	}
@@ -372,7 +458,7 @@ FBox FFastGeoInstancedStaticMeshComponent::GetNavigationBounds() const
 
 bool FFastGeoInstancedStaticMeshComponent::IsNavigationRelevant() const
 {
-	return PerInstanceSMData.Num() > 0 && Super::IsNavigationRelevant();
+	return GetInstanceCount() > 0 && Super::IsNavigationRelevant();
 }
 
 bool FFastGeoInstancedStaticMeshComponent::DoCustomNavigableGeometryExport(FNavigableGeometryExport& GeomExport) const
diff --git a/Engine/Plugins/Experimental/FastGeoStreaming/Source/FastGeoStreaming/Private/FastGeoWorldPartitionRuntimeCellTransformer.cpp b/Engine/Plugins/Experimental/FastGeoStreaming/Source/FastGeoStreaming/Private/FastGeoWorldPartitionRuntimeCellTransformer.cpp
--- a/Engine/Plugins/Experimental/FastGeoStreaming/Source/FastGeoStreaming/Private/FastGeoWorldPartitionRuntimeCellTransformer.cpp
+++ b/Engine/Plugins/Experimental/FastGeoStreaming/Source/FastGeoStreaming/Private/FastGeoWorldPartitionRuntimeCellTransformer.cpp
@@ -131,6 +131,8 @@ namespace FastGeo
 		return FFastGeoElementType::Invalid;
 	}
 }
+
+UFastGeoWorldPartitionRuntimeCellTransformer* UFastGeoWorldPartitionRuntimeCellTransformer::CurrentTransformer = nullptr;
 #endif
 
 void UFastGeoTransformerSettings::PostLoad()
@@ -313,13 +315,15 @@ void UFastGeoWorldPartitionRuntimeCellTransformer::Transform(ULevel* InLevel)
 	}
 
 	check(InLevel);
+	check(!UFastGeoWorldPartitionRuntimeCellTransformer::CurrentTransformer);
+	TGuardValue<UFastGeoWorldPartitionRuntimeCellTransformer*> GuardCurrentTransformer(UFastGeoWorldPartitionRuntimeCellTransformer::CurrentTransformer, this);
 	
 	if (IsDebugMode())
 	{
 		UE_LOG(LogFastGeoStreaming, Log, TEXT("------------------------------------------------------------------------"));
 		UE_LOG(LogFastGeoStreaming, Log, TEXT("- FastGeoStreaming Debug Mode: Transforming Level '%s'"), *InLevel->GetPathName());
 	}
-	
+
 	FTransformationStats Stats;
 	TMap<AActor*, FTransformableActor> TransformableActors;
 	GatherTransformableActors(InLevel->Actors, InLevel, TransformableActors, Stats);
@@ -417,9 +421,9 @@ void UFastGeoWorldPartitionRuntimeCellTransformer::ForEachIgnoredComponentClass(
 {
 	Super::ForEachIgnoredComponentClass(Func);
 
-	const UFastGeoTransformerSettings* EffectiveSettings = GetSettings();
+	const UFastGeoTransformerSettings& EffectiveSettings = GetSettings();
 
-	for (const TSubclassOf<UActorComponent>& IgnoredRemainingComponentClass : EffectiveSettings->IgnoredRemainingComponentClasses)
+	for (const TSubclassOf<UActorComponent>& IgnoredRemainingComponentClass : EffectiveSettings.IgnoredRemainingComponentClasses)
 	{
 		if (!Func(IgnoredRemainingComponentClass))
 		{
@@ -439,9 +443,9 @@ void UFastGeoWorldPartitionRuntimeCellTransformer::ForEachIgnoredExactComponentC
 {
 	Super::ForEachIgnoredExactComponentClass(Func);
 
-	const UFastGeoTransformerSettings* EffectiveSettings = GetSettings();
+	const UFastGeoTransformerSettings& EffectiveSettings = GetSettings();
 
-	for (const TSubclassOf<UActorComponent>& IgnoredRemainingExactComponentClass : EffectiveSettings->IgnoredRemainingExactComponentClasses)
+	for (const TSubclassOf<UActorComponent>& IgnoredRemainingExactComponentClass : EffectiveSettings.IgnoredRemainingExactComponentClasses)
 	{
 		if (!Func(IgnoredRemainingExactComponentClass))
 		{
@@ -614,10 +618,10 @@ void UFastGeoWorldPartitionRuntimeCellTransformer::FTransformationStats::DumpSta
 
 FFastGeoTransformResult UFastGeoWorldPartitionRuntimeCellTransformer::IsAllowedActorClass(AActor* InActor) const
 {
-	const UFastGeoTransformerSettings* EffectiveSettings = GetSettings();
+	const UFastGeoTransformerSettings& EffectiveSettings = GetSettings();
 	UClass* ActorClass = InActor->GetClass();
 
-	for (TSubclassOf<AActor> DisallowedActorClass : EffectiveSettings->DisallowedActorClasses)
+	for (TSubclassOf<AActor> DisallowedActorClass : EffectiveSettings.DisallowedActorClasses)
 	{
 		if (ActorClass->IsChildOf(DisallowedActorClass))
 		{
@@ -633,7 +637,7 @@ FFastGeoTransformResult UFastGeoWorldPartitionRuntimeCellTransformer::IsAllowedA
 		}
 	}
 
-	for (TSubclassOf<AActor> DisallowedExactActorClass : EffectiveSettings->DisallowedExactActorClasses)
+	for (TSubclassOf<AActor> DisallowedExactActorClass : EffectiveSettings.DisallowedExactActorClasses)
 	{
 		if (ActorClass == DisallowedExactActorClass)
 		{
@@ -641,7 +645,7 @@ FFastGeoTransformResult UFastGeoWorldPartitionRuntimeCellTransformer::IsAllowedA
 		}
 	}
 
-	for (TSubclassOf<AActor> AllowedActorClass : EffectiveSettings->AllowedActorClasses)
+	for (TSubclassOf<AActor> AllowedActorClass : EffectiveSettings.AllowedActorClasses)
 	{
 		if (ActorClass->IsChildOf(AllowedActorClass))
 		{
@@ -657,7 +661,7 @@ FFastGeoTransformResult UFastGeoWorldPartitionRuntimeCellTransformer::IsAllowedA
 		}
 	}
 
-	for (TSubclassOf<AActor> AllowedExactActorClass : EffectiveSettings->AllowedExactActorClasses)
+	for (TSubclassOf<AActor> AllowedExactActorClass : EffectiveSettings.AllowedExactActorClasses)
 	{
 		if (ActorClass == AllowedExactActorClass)
 		{
@@ -817,7 +821,7 @@ FFastGeoTransformResult UFastGeoWorldPartitionRuntimeCellTransformer::CanTransfo
 
 FFastGeoTransformResult UFastGeoWorldPartitionRuntimeCellTransformer::IsAllowedComponentClass(UActorComponent* InComponent) const
 {
-	const UFastGeoTransformerSettings* EffectiveSettings = GetSettings();
+	const UFastGeoTransformerSettings& EffectiveSettings = GetSettings();
 
 	FFastGeoElementType FastGeoComponentType = FastGeo::GetFastGeoComponentType(InComponent->GetClass());
 	if (!FastGeoComponentType.IsValid())
@@ -826,7 +830,7 @@ FFastGeoTransformResult UFastGeoWorldPartitionRuntimeCellTransformer::IsAllowedC
 	}
 
 	UClass* ComponentClass = InComponent->GetClass();
-	for (TSubclassOf<UActorComponent> DisallowedComponentClass : EffectiveSettings->DisallowedComponentClasses)
+	for (TSubclassOf<UActorComponent> DisallowedComponentClass : EffectiveSettings.DisallowedComponentClasses)
 	{
 		if (ComponentClass->IsChildOf(DisallowedComponentClass))
 		{
@@ -842,7 +846,7 @@ FFastGeoTransformResult UFastGeoWorldPartitionRuntimeCellTransformer::IsAllowedC
 		}
 	}
 
-	for (TSubclassOf<UActorComponent> DisallowedExactComponentClass : EffectiveSettings->DisallowedExactComponentClasses)
+	for (TSubclassOf<UActorComponent> DisallowedExactComponentClass : EffectiveSettings.DisallowedExactComponentClasses)
 	{
 		if (ComponentClass == DisallowedExactComponentClass)
 		{
@@ -850,7 +854,7 @@ FFastGeoTransformResult UFastGeoWorldPartitionRuntimeCellTransformer::IsAllowedC
 		}
 	}
 
-	for (TSubclassOf<UActorComponent> AllowedComponentClass : EffectiveSettings->AllowedComponentClasses)
+	for (TSubclassOf<UActorComponent> AllowedComponentClass : EffectiveSettings.AllowedComponentClasses)
 	{
 		if (ComponentClass->IsChildOf(AllowedComponentClass))
 		{
@@ -866,7 +870,7 @@ FFastGeoTransformResult UFastGeoWorldPartitionRuntimeCellTransformer::IsAllowedC
 		}
 	}
 	
-	for (TSubclassOf<UActorComponent> AllowedExactComponentClass : EffectiveSettings->AllowedExactComponentClasses)
+	for (TSubclassOf<UActorComponent> AllowedExactComponentClass : EffectiveSettings.AllowedExactComponentClasses)
 	{
 		if (ComponentClass == AllowedExactComponentClass)
 		{
diff --git a/Engine/Plugins/Experimental/FastGeoStreaming/Source/FastGeoStreaming/Public/FastGeoWorldPartitionRuntimeCellTransformer.h b/Engine/Plugins/Experimental/FastGeoStreaming/Source/FastGeoStreaming/Public/FastGeoWorldPartitionRuntimeCellTransformer.h
--- a/Engine/Plugins/Experimental/FastGeoStreaming/Source/FastGeoStreaming/Public/FastGeoWorldPartitionRuntimeCellTransformer.h
+++ b/Engine/Plugins/Experimental/FastGeoStreaming/Source/FastGeoStreaming/Public/FastGeoWorldPartitionRuntimeCellTransformer.h
@@ -41,6 +41,47 @@ private:
 	EFastGeoTransform TransformResult;
 };
 
+/** Defines the memory storage format for FastGeo instance transforms. */
+UENUM()
+enum class EFastGeoInstanceStorageMode : uint8
+{
+	/** 
+ 	 * Automatically selects the storage format by checking if compression causes
+ 	 * precision loss in position, rotation, or scale beyond the allowed thresholds.
+ 	 */
+	Auto,
+	/** 
+	 * Forces using FTransform3f (recommended for most game levels to optimize memory footprint).
+	 */
+	Compressed,
+	/** 
+	 * Forces using FMatrix (Best for absolute accuracy when instances show jitter or scale artifacts).
+	 */
+	Full
+};
+
+/** Threshold values used when transforming InstanceStaticMeshComponent tranforms when 
+  * FastGeo transformer InstanceStorageMode is set to EFastGeoInstanceStorageMode::Auto.
+  */
+USTRUCT()
+struct FASTGEOSTREAMING_API FFastGeoInstanceCompressionSettings
+{
+	GENERATED_BODY()
+
+public:
+	/** Threshold for Position (cm) */
+	UPROPERTY(EditAnywhere, Category = "Instance Precision", meta = (ClampMin = "0.0", UIMin = "0.0", UIMax = "1.0", Step = "0.001"))
+	float PositionEpsilon = 5e-2f;
+
+	/** Threshold for Rotation */
+	UPROPERTY(EditAnywhere, Category = "Instance Precision", meta = (ClampMin = "0.0", UIMin = "0.0", UIMax = "0.01", Step = "0.0001"))
+	float RotationEpsilon = 1e-4f;
+
+	/** Threshold for Scale */
+	UPROPERTY(EditAnywhere, Category = "Instance Precision", meta = (ClampMin = "0.0", UIMin = "0.0", UIMax = "0.01", Step = "0.0001"))
+	float ScaleEpsilon = 1e-4f;
+};
+
 // Empty type used to present a Button in the property grid - not actual data
 USTRUCT()
 struct FFastGeoConversionButton
@@ -98,6 +139,13 @@ public:
 	UPROPERTY(EditAnywhere, Category="Settings", meta = (AllowAbstract = "true"))
 	TArray<TSubclassOf<UActorComponent>> IgnoredRemainingExactComponentClasses;
 
+	/** Determines how instance transforms are stored in memory */
+	UPROPERTY(EditAnywhere, Category = "Settings")
+	EFastGeoInstanceStorageMode InstanceStorageMode = EFastGeoInstanceStorageMode::Full;
+
+	UPROPERTY(EditAnywhere, Category = "Settings", meta = (EditCondition = "InstanceStorageMode == EFastGeoInstanceStorageMode::Auto", EditConditionHides))
+	FFastGeoInstanceCompressionSettings InstanceCompressionSettings;
+
 	UPROPERTY(VisibleAnywhere, Category = "Settings", meta = (DisplayName = "Convert settings to reusable Asset"))
 	FFastGeoConversionButton FastGeoConversionButton;
 };
@@ -126,6 +174,9 @@ public:
 	virtual void PostEditChangeProperty(FPropertyChangedEvent& InPropertyChangedEvent) override;
 	//~ End UObject Interface.
 
+	/** Used internally to get transformer settings. */
+	static const UFastGeoWorldPartitionRuntimeCellTransformer* GetCurrentTransformer() { return CurrentTransformer; }
+
 protected:
 	/** Whether the actor can be processed by the transformer. */
 	virtual bool IsActorTransformable(AActor* InActor, FString& OutReason) const { return true; }
@@ -171,18 +222,19 @@ public:
 #if WITH_EDITORONLY_DATA
 	UPROPERTY(EditAnywhere, Category = "Settings")
 	TObjectPtr<UFastGeoTransformerSettings> Settings = nullptr;
+
+	const UFastGeoTransformerSettings& GetSettings() const
+	{
+		check(EmbeddedSettings);
+		return (Settings != nullptr) ? *Settings : *EmbeddedSettings;
+	}
 #endif
 
 protected:
 #if WITH_EDITORONLY_DATA
-	UPROPERTY(EditAnywhere, Instanced, Meta = (EditCondition = "Settings == nullptr", Tooltip = "Fallback only available when no Settings asset is used."), Category = "Settings")
+	UPROPERTY(EditAnywhere, Instanced, NoClear, Meta = (EditCondition = "Settings == nullptr", Tooltip = "Fallback only available when no Settings asset is used."), Category = "Settings")
 	TObjectPtr<UFastGeoTransformerSettings> EmbeddedSettings;
 	
-	const UFastGeoTransformerSettings* GetSettings() const
-	{
-		return (Settings != nullptr) ? Settings : EmbeddedSettings;
-	}
-
 	UPROPERTY(EditAnywhere, Transient, SkipSerialization, Category = "Debug")
 	bool bDebugMode;
 
@@ -247,4 +299,7 @@ protected:
 	UPROPERTY(Config)
 	TArray<TSubclassOf<UActorComponent>> BuiltinIgnoredRemainingExactComponentClasses;
 #endif
+
+private:
+	static UFastGeoWorldPartitionRuntimeCellTransformer* CurrentTransformer;
 };
\ No newline at end of file
diff --git a/Engine/Source/Runtime/Engine/Classes/Components/InstancedStaticMeshComponent.h b/Engine/Source/Runtime/Engine/Classes/Components/InstancedStaticMeshComponent.h
--- a/Engine/Source/Runtime/Engine/Classes/Components/InstancedStaticMeshComponent.h
+++ b/Engine/Source/Runtime/Engine/Classes/Components/InstancedStaticMeshComponent.h
@@ -455,6 +455,9 @@ public:
 	/** Returns true if RemoveAtSwap is enabled. Derived classes may always return true regardless of whether SetRemoveSwapEnabled() was called. */
 	ENGINE_API virtual bool SupportsRemoveSwap() const;
 
+	/** Iterates on instance transform (FMatrix) */
+	void ForEachInstanceMatrix(TFunctionRef<void(const FMatrix&)> Func) const;
+
 	/** 
 	 * Sets whether to use conservative bounds. 
 	 * This doesn't fully recalculate bounds on instance addition/removal/transform change.
diff --git a/Engine/Source/Runtime/Engine/Private/InstancedStaticMesh.cpp b/Engine/Source/Runtime/Engine/Private/InstancedStaticMesh.cpp
--- a/Engine/Source/Runtime/Engine/Private/InstancedStaticMesh.cpp
+++ b/Engine/Source/Runtime/Engine/Private/InstancedStaticMesh.cpp
@@ -4941,6 +4941,14 @@ bool UInstancedStaticMeshComponent::IsValidInstance(int32 InstanceIndex) const
 	return PerInstanceSMData.IsValidIndex(InstanceIndex);
 }
 
+void UInstancedStaticMeshComponent::ForEachInstanceMatrix(TFunctionRef<void(const FMatrix&)> Func) const
+{
+	for (const FInstancedStaticMeshInstanceData& InstanceData : PerInstanceSMData)
+	{
+		Func(InstanceData.Transform);
+	}
+}
+
 void UInstancedStaticMeshComponent::SetCullDistances(int32 StartCullDistance, int32 EndCullDistance)
 {
 	if (InstanceStartCullDistance != StartCullDistance || InstanceEndCullDistance != EndCullDistance)
diff --git a/Engine/Source/Runtime/Engine/Public/InstancedStaticMeshComponentHelper.h b/Engine/Source/Runtime/Engine/Public/InstancedStaticMeshComponentHelper.h
--- a/Engine/Source/Runtime/Engine/Public/InstancedStaticMeshComponentHelper.h
+++ b/Engine/Source/Runtime/Engine/Public/InstancedStaticMeshComponentHelper.h
@@ -151,18 +151,18 @@ void FInstancedStaticMeshComponentHelper::GetNavigationPerInstanceTransforms(con
 	if (InstanceBounds.IsValid)
 	{
 		const FBox LocalAreaBox = AreaBox.InverseTransformBy(Component.GetComponentTransform());
-		for (const auto& InstancedData : Component.PerInstanceSMData)
+		Component.ForEachInstanceMatrix([&Component, &InstanceBounds, &LocalAreaBox, &InstanceData](const FMatrix& InstanceTransform)
 		{
-			const FTransform InstanceToComponent(InstancedData.Transform);
+			const FTransform InstanceToComponent(InstanceTransform);
 			if (!InstanceToComponent.GetScale3D().IsZero())
 			{
-				const FBoxSphereBounds TransformedInstanceBounds = InstanceBounds.TransformBy(InstancedData.Transform);
+				const FBoxSphereBounds TransformedInstanceBounds = InstanceBounds.TransformBy(InstanceTransform);
 				if (LocalAreaBox.Intersect(TransformedInstanceBounds.GetBox()))
 				{
 					InstanceData.Add(InstanceToComponent * Component.GetComponentTransform());
 				}
 			}
-		}
+		});
 	}
 }
 
