Fixed FBodyInstance::GetSimplePhysicalMaterial and FBodyInstance::GetComplexPhysicalMaterials to forward calls to the BodyInstanceOwner when valid. This fix is a prerequisite for an upcoming feature.
Deprecated IPhysicsBodyInstanceOwner::GetPhysicalMaterial and replaced by GetSimplePhysicalMaterial.

#Synced-CL 49829988

diff --git a/Engine/Plugins/Experimental/FastGeoStreaming/Source/FastGeoStreaming/Internal/FastGeoPhysicsBodyInstanceOwner.h b/Engine/Plugins/Experimental/FastGeoStreaming/Source/FastGeoStreaming/Internal/FastGeoPhysicsBodyInstanceOwner.h
--- a/Engine/Plugins/Experimental/FastGeoStreaming/Source/FastGeoStreaming/Internal/FastGeoPhysicsBodyInstanceOwner.h
+++ b/Engine/Plugins/Experimental/FastGeoStreaming/Source/FastGeoStreaming/Internal/FastGeoPhysicsBodyInstanceOwner.h
@@ -28,7 +28,7 @@ public:
 	virtual ECollisionChannel GetCollisionObjectType() const override;
 	virtual ECollisionEnabled::Type GetCollisionEnabled() const override;
 	virtual ECollisionResponse GetCollisionResponseToChannel(ECollisionChannel Channel) const override;
-	virtual UPhysicalMaterial* GetPhysicalMaterial() const override;
+	virtual UPhysicalMaterial* GetSimplePhysicalMaterial(const TWeakObjectPtr<UBodySetup>& InBodySetup) const override;
 	virtual void GetComplexPhysicalMaterials(TArray<UPhysicalMaterial*>& OutPhysMaterials, TArray<FPhysicalMaterialMaskParams>* OutPhysMaterialMasks) const override;
 	virtual Chaos::FPhysicsObject* GetPhysicsObjectById(Chaos::FPhysicsObjectId Id) const override;
 	virtual TArray<Chaos::FPhysicsObject*> GetAllPhysicsObjects() const override;
diff --git a/Engine/Plugins/Experimental/FastGeoStreaming/Source/FastGeoStreaming/Internal/FastGeoPrimitiveComponent.h b/Engine/Plugins/Experimental/FastGeoStreaming/Source/FastGeoStreaming/Internal/FastGeoPrimitiveComponent.h
--- a/Engine/Plugins/Experimental/FastGeoStreaming/Source/FastGeoStreaming/Internal/FastGeoPrimitiveComponent.h
+++ b/Engine/Plugins/Experimental/FastGeoStreaming/Source/FastGeoStreaming/Internal/FastGeoPrimitiveComponent.h
@@ -152,7 +152,9 @@ protected:
 	virtual UObject* GetSourceObject() const;
 	virtual ECollisionChannel GetCollisionObjectType() const;
 	virtual ECollisionEnabled::Type GetCollisionEnabled() const;
+	UE_DEPRECATED(5.8, "Use GetSimplePhysicalMaterial instead.")
 	virtual UPhysicalMaterial* GetPhysicalMaterial() const { return nullptr; }
+	virtual UPhysicalMaterial* GetSimplePhysicalMaterial(const TWeakObjectPtr<UBodySetup>& InBodySetup) const { return nullptr; }
 	virtual void GetComplexPhysicalMaterials(TArray<UPhysicalMaterial*>& OutPhysMaterials, TArray<FPhysicalMaterialMaskParams>* OutPhysMaterialMasks) const {}
 	virtual ECollisionResponse GetCollisionResponseToChannel(ECollisionChannel Channel) const;
 	virtual Chaos::FPhysicsObject* GetPhysicsObjectById(Chaos::FPhysicsObjectId Id) const;
diff --git a/Engine/Plugins/Experimental/FastGeoStreaming/Source/FastGeoStreaming/Internal/FastGeoStaticMeshComponent.h b/Engine/Plugins/Experimental/FastGeoStreaming/Source/FastGeoStreaming/Internal/FastGeoStaticMeshComponent.h
--- a/Engine/Plugins/Experimental/FastGeoStreaming/Source/FastGeoStreaming/Internal/FastGeoStaticMeshComponent.h
+++ b/Engine/Plugins/Experimental/FastGeoStreaming/Source/FastGeoStreaming/Internal/FastGeoStaticMeshComponent.h
@@ -93,7 +93,7 @@ protected:
 	virtual void InitializeSceneProxyDescDynamicProperties() override;
 	virtual void ApplyWorldTransform(const FTransform& InTransform) override;
 	virtual FPrimitiveSceneProxy* CreateSceneProxy(ESceneProxyCreationError* OutError = nullptr) override;
-	virtual UPhysicalMaterial* GetPhysicalMaterial() const override;
+	virtual UPhysicalMaterial* GetSimplePhysicalMaterial(const TWeakObjectPtr<UBodySetup>& InBodySetup) const override;
 	virtual void GetComplexPhysicalMaterials(TArray<UPhysicalMaterial*>& OutPhysMaterials, TArray<FPhysicalMaterialMaskParams>* OutPhysMaterialMasks) const override;
 	//~ End FFastGeoPrimitiveComponent interface
 
diff --git a/Engine/Plugins/Experimental/FastGeoStreaming/Source/FastGeoStreaming/Private/FastGeoPhysicsBodyInstanceOwner.cpp b/Engine/Plugins/Experimental/FastGeoStreaming/Source/FastGeoStreaming/Private/FastGeoPhysicsBodyInstanceOwner.cpp
--- a/Engine/Plugins/Experimental/FastGeoStreaming/Source/FastGeoStreaming/Private/FastGeoPhysicsBodyInstanceOwner.cpp
+++ b/Engine/Plugins/Experimental/FastGeoStreaming/Source/FastGeoStreaming/Private/FastGeoPhysicsBodyInstanceOwner.cpp
@@ -78,10 +78,10 @@ ECollisionResponse FFastGeoPhysicsBodyInstanceOwner::GetCollisionResponseToChann
 	return OwnerComponent->GetCollisionResponseToChannel(Channel);
 }
 
-UPhysicalMaterial* FFastGeoPhysicsBodyInstanceOwner::GetPhysicalMaterial() const
+UPhysicalMaterial* FFastGeoPhysicsBodyInstanceOwner::GetSimplePhysicalMaterial(const TWeakObjectPtr<UBodySetup>& InBodySetup) const
 {
 	check(OwnerContainer.IsValid());
-	return OwnerComponent->GetPhysicalMaterial();
+	return OwnerComponent->GetSimplePhysicalMaterial(InBodySetup);
 }
 
 void FFastGeoPhysicsBodyInstanceOwner::GetComplexPhysicalMaterials(TArray<UPhysicalMaterial*>& OutPhysMaterials, TArray<FPhysicalMaterialMaskParams>* OutPhysMaterialMasks) const
diff --git a/Engine/Plugins/Experimental/FastGeoStreaming/Source/FastGeoStreaming/Private/FastGeoStaticMeshComponent.cpp b/Engine/Plugins/Experimental/FastGeoStreaming/Source/FastGeoStreaming/Private/FastGeoStaticMeshComponent.cpp
--- a/Engine/Plugins/Experimental/FastGeoStreaming/Source/FastGeoStreaming/Private/FastGeoStaticMeshComponent.cpp
+++ b/Engine/Plugins/Experimental/FastGeoStreaming/Source/FastGeoStreaming/Private/FastGeoStaticMeshComponent.cpp
@@ -14,6 +14,7 @@
 #include "StaticMeshComponentHelper.h"
 #include "Rendering/NaniteResourcesHelper.h"
 #include "PhysicalMaterials/PhysicalMaterial.h"
+#include "PhysicsEngine/BodySetup.h"
 #include "NaniteVertexFactory.h"
 #include "StaticMeshSceneProxy.h"
 #include "VT/MeshPaintVirtualTexture.h"
@@ -464,13 +465,46 @@ bool FFastGeoStaticMeshComponentBase::IsDisplayNaniteFallbackMesh() const
 #endif
 
 //~ Begin IPhysicsBodyInstanceOwner interface
-UPhysicalMaterial* FFastGeoStaticMeshComponentBase::GetPhysicalMaterial() const
+UPhysicalMaterial* FFastGeoStaticMeshComponentBase::GetSimplePhysicalMaterial(const TWeakObjectPtr<UBodySetup>& InBodySetup) const
 {
-	if (UMaterialInterface* Material = GetMaterial(0))
+	if (!GEngine || !GEngine->DefaultPhysMaterial)
 	{
-		return Material->GetPhysicalMaterial();
+		UE_LOG(LogPhysics, Error, TEXT("FFastGeoStaticMeshComponentBase::GetSimplePhysicalMaterial : GEngine not initialized!"));
+		return nullptr;
 	}
-	return nullptr;
+
+	UPhysicalMaterial* ReturnPhysMaterial = nullptr;
+	if (BodyInstance.GetPhysMaterialOverride())
+	{
+		ReturnPhysMaterial = BodyInstance.GetPhysMaterialOverride();
+	}
+	else
+	{
+		// Use InBodySetup if possible
+		UBodySetup* BodySetup = InBodySetup.Get();
+		if (BodySetup && BodySetup->PhysMaterial)
+		{
+			ReturnPhysMaterial = BodySetup->PhysMaterial;
+		}
+		else
+		{
+			// See if the Material has a PhysicalMaterial
+			UMaterialInterface* Material = GetMaterial(0);
+			UPhysicalMaterial* PhysMatFromMaterial = Material ? Material->GetPhysicalMaterial() : nullptr;
+			if (PhysMatFromMaterial)
+			{
+				ReturnPhysMaterial = PhysMatFromMaterial;
+			}
+			// Fallback is default physical material
+			else
+			{
+				ReturnPhysMaterial = GEngine->DefaultPhysMaterial;
+			}
+		}
+	}
+
+	check(!ReturnPhysMaterial || ReturnPhysMaterial->IsValidLowLevel());
+	return ReturnPhysMaterial;
 }
 
 void FFastGeoStaticMeshComponentBase::GetComplexPhysicalMaterials(TArray<UPhysicalMaterial*>& OutPhysMaterials, TArray<FPhysicalMaterialMaskParams>* OutPhysMaterialMasks) const
diff --git a/Engine/Source/Runtime/Engine/Classes/PhysicsEngine/PhysicsBodyInstanceOwnerInterface.h b/Engine/Source/Runtime/Engine/Classes/PhysicsEngine/PhysicsBodyInstanceOwnerInterface.h
--- a/Engine/Source/Runtime/Engine/Classes/PhysicsEngine/PhysicsBodyInstanceOwnerInterface.h
+++ b/Engine/Source/Runtime/Engine/Classes/PhysicsEngine/PhysicsBodyInstanceOwnerInterface.h
@@ -9,8 +9,10 @@
 #include "PhysicsBodyInstanceOwnerInterface.generated.h"
 
 class UObject;
+class UBodySetup;
 class UPhysicalMaterial;
 class IPhysicsBodyInstanceOwner;
+struct FBodyInstance;
 struct FHitResult;
 struct FOverlapResult;
 struct FPhysicalMaterialMaskParams;
@@ -70,8 +72,11 @@ public:
 	/** Returns the form of collision for this body instance owner. */
 	virtual ECollisionEnabled::Type GetCollisionEnabled() const = 0;
 
+	UE_DEPRECATED(5.8, "Use GetSimplePhysicalMaterial instead.")
+	virtual UPhysicalMaterial* GetPhysicalMaterial() const { return GetSimplePhysicalMaterial(nullptr); };
+
 	/** Find the correct PhysicalMaterial for simple geometry on this body (used by FBodyInstance::GetSimplePhysicalMaterial). */
-	virtual UPhysicalMaterial* GetPhysicalMaterial() const = 0;
+	virtual UPhysicalMaterial* GetSimplePhysicalMaterial(const TWeakObjectPtr<UBodySetup>& InBodySetup) const = 0;
 
 	/** Gets the complex PhysicalMaterials array for this body (used by FBodyInstance::GetComplexPhysicalMaterials). */
 	virtual void GetComplexPhysicalMaterials(TArray<UPhysicalMaterial*>& OutPhysMaterials, TArray<FPhysicalMaterialMaskParams>* OutPhysMaterialMasks) const = 0;
diff --git a/Engine/Source/Runtime/Engine/Private/PhysicsEngine/BodyInstance.cpp b/Engine/Source/Runtime/Engine/Private/PhysicsEngine/BodyInstance.cpp
--- a/Engine/Source/Runtime/Engine/Private/PhysicsEngine/BodyInstance.cpp
+++ b/Engine/Source/Runtime/Engine/Private/PhysicsEngine/BodyInstance.cpp
@@ -3163,6 +3163,11 @@ UPhysicalMaterial* FBodyInstance::GetSimplePhysicalMaterial(const FBodyInstance*
 		ReturnPhysMaterial = BodyInstance->PhysMaterialOverride;
 		check(!ReturnPhysMaterial || ReturnPhysMaterial->IsValidLowLevel());
 	}
+	else if (BodyInstanceOwner)
+	{
+		ReturnPhysMaterial = BodyInstanceOwner->GetSimplePhysicalMaterial(BodySetupPtr);
+		check(!ReturnPhysMaterial || ReturnPhysMaterial->IsValidLowLevel());
+	}
 	else
 	{
 		// Component override
@@ -3194,10 +3199,6 @@ UPhysicalMaterial* FBodyInstance::GetSimplePhysicalMaterial(const FBodyInstance*
 						PhysMatFromMaterial = Material->GetPhysicalMaterial();
 					}
 				}
-				else if (BodyInstanceOwner)
-				{
-					PhysMatFromMaterial = BodyInstanceOwner->GetPhysicalMaterial();
-				}
 
 				if (PhysMatFromMaterial != NULL)
 				{
@@ -3262,8 +3263,13 @@ void FBodyInstance::GetComplexPhysicalMaterials(const FBodyInstance* BodyInstanc
 	}
 	else
 	{
+		// Use BodyInstanceOwner if valid
+		if (InBodyInstanceOwner)
+		{
+			InBodyInstanceOwner->GetComplexPhysicalMaterials(OutPhysMaterials, OutPhysMaterialMasks);
+		}
 		// Component override
-		if (UPrimitiveComponent* OwnerPrimComponent = OwnerComp.Get())
+		else if (UPrimitiveComponent* OwnerPrimComponent = OwnerComp.Get())
 		{
 			if (OwnerPrimComponent->BodyInstance.PhysMaterialOverride != nullptr)
 			{
@@ -3313,10 +3319,6 @@ void FBodyInstance::GetComplexPhysicalMaterials(const FBodyInstance* BodyInstanc
 				}
 			}
 		}
-		else if (InBodyInstanceOwner)
-		{
-			InBodyInstanceOwner->GetComplexPhysicalMaterials(OutPhysMaterials, OutPhysMaterialMasks);
-		}
 	}
 }
 
