[FastGeo] Custom HLOD support
* Added support for Custom HLOD actors
* Refactored StandaloneHLODGuid and CustomHLODGuid into a single GUID and enum to reduce memory usage

#Synced-CL 50263503

diff --git a/Engine/Plugins/Experimental/FastGeoStreaming/Source/FastGeoStreaming/Internal/FastGeoHLOD.h b/Engine/Plugins/Experimental/FastGeoStreaming/Source/FastGeoStreaming/Internal/FastGeoHLOD.h
--- a/Engine/Plugins/Experimental/FastGeoStreaming/Source/FastGeoStreaming/Internal/FastGeoHLOD.h
+++ b/Engine/Plugins/Experimental/FastGeoStreaming/Source/FastGeoStreaming/Internal/FastGeoHLOD.h
@@ -58,6 +58,16 @@ private:
 	// Persistent Data
 	bool bRequireWarmup;
 	FGuid SourceCellGuid;
-	FGuid StandaloneHLODGuid;
-	FGuid CustomHLODGuid;
+
+	enum class EHLODGuidType : uint8
+	{
+		None,
+		Standalone,
+		Custom
+	};
+	FGuid HLODGuid;
+	EHLODGuidType HLODGuidType = EHLODGuidType::None;
+
+	const FGuid& GetHLODGuid(EHLODGuidType InGuidType) const;
+	void SetHLODGuid(const FGuid& InGuid, EHLODGuidType InGuidType);
 };
diff --git a/Engine/Plugins/Experimental/FastGeoStreaming/Source/FastGeoStreaming/Private/FastGeoHLOD.cpp b/Engine/Plugins/Experimental/FastGeoStreaming/Source/FastGeoStreaming/Private/FastGeoHLOD.cpp
--- a/Engine/Plugins/Experimental/FastGeoStreaming/Source/FastGeoStreaming/Private/FastGeoHLOD.cpp
+++ b/Engine/Plugins/Experimental/FastGeoStreaming/Source/FastGeoStreaming/Private/FastGeoHLOD.cpp
@@ -31,8 +31,8 @@ void FFastGeoHLOD::Serialize(FArchive& Ar)
 
 	Ar << bRequireWarmup;
 	Ar << SourceCellGuid;
-	Ar << StandaloneHLODGuid;
-	Ar << CustomHLODGuid;
+	Ar << HLODGuid;
+	Ar << HLODGuidType;
 }
 
 ULevel* FFastGeoHLOD::GetHLODLevel() const
@@ -105,22 +105,35 @@ const FGuid& FFastGeoHLOD::GetSourceCellGuid() const
 
 bool FFastGeoHLOD::IsStandalone() const
 {
-	return StandaloneHLODGuid.IsValid();
+	return GetStandaloneHLODGuid().IsValid();
 }
 
 const FGuid& FFastGeoHLOD::GetStandaloneHLODGuid() const
 {
-	return StandaloneHLODGuid;
+	return GetHLODGuid(EHLODGuidType::Standalone);
 }
 
 bool FFastGeoHLOD::IsCustomHLOD() const
 {
-	return CustomHLODGuid.IsValid();
+	return GetCustomHLODGuid().IsValid();
 }
 
 const FGuid& FFastGeoHLOD::GetCustomHLODGuid() const
 {
-	return CustomHLODGuid;
+	return GetHLODGuid(EHLODGuidType::Custom);
+}
+
+const FGuid& FFastGeoHLOD::GetHLODGuid(EHLODGuidType InGuidType) const
+{
+	if (HLODGuidType == InGuidType)
+	{
+		return HLODGuid;
+	}
+	else
+	{
+		static const FGuid InvalidGuid;
+		return InvalidGuid;
+	}
 }
 
 void FFastGeoHLOD::OnRegister()
@@ -146,12 +159,18 @@ void FFastGeoHLOD::SetRequireWarmup(bool bInRequireWarmup)
 
 void FFastGeoHLOD::SetStandaloneHLODGuid(const FGuid& InStandaloneHLODGuid)
 {
-	StandaloneHLODGuid = InStandaloneHLODGuid;
+	SetHLODGuid(InStandaloneHLODGuid, EHLODGuidType::Standalone);
 }
 
 void FFastGeoHLOD::SetCustomHLODGuid(const FGuid& InCustomHLODGuid)
 {
-	CustomHLODGuid = InCustomHLODGuid;
+	SetHLODGuid(InCustomHLODGuid, EHLODGuidType::Custom);
+}
+
+void FFastGeoHLOD::SetHLODGuid(const FGuid& InGuid, EHLODGuidType InGuidType)
+{
+	HLODGuid = InGuid;
+	HLODGuidType = InGuidType;
 }
 
 FFastGeoComponent& FFastGeoHLOD::AddComponent(FFastGeoElementType InComponentType)
diff --git a/Engine/Plugins/Experimental/FastGeoStreaming/Source/FastGeoStreaming/Private/FastGeoWorldPartitionRuntimeCellTransformer.cpp b/Engine/Plugins/Experimental/FastGeoStreaming/Source/FastGeoStreaming/Private/FastGeoWorldPartitionRuntimeCellTransformer.cpp
--- a/Engine/Plugins/Experimental/FastGeoStreaming/Source/FastGeoStreaming/Private/FastGeoWorldPartitionRuntimeCellTransformer.cpp
+++ b/Engine/Plugins/Experimental/FastGeoStreaming/Source/FastGeoStreaming/Private/FastGeoWorldPartitionRuntimeCellTransformer.cpp
@@ -49,6 +49,7 @@
 
 #if WITH_EDITORONLY_DATA
 #include "WorldPartition/HLOD/HLODActor.h"
+#include "WorldPartition/HLOD/CustomHLODActor.h"
 #endif
 
 #define LOCTEXT_NAMESPACE "WorldPartition"
@@ -65,9 +66,27 @@ namespace FastGeo
 	// Tag use to force exclude actors from FastGeoStreaming
 	static const FName NAME_NoFastGeo(TEXT("NoFastGeo"));
 
+	static IWorldPartitionHLODObject* GetHLODObject(AActor* Actor)
+	{
+		if (AWorldPartitionHLOD* HLODActor = Cast<AWorldPartitionHLOD>(Actor))
+		{
+			return HLODActor;
+		}
+		if (AWorldPartitionCustomHLOD* CustomHLODActor = Cast<AWorldPartitionCustomHLOD>(Actor))
+		{
+			return CustomHLODActor;
+		}
+		return nullptr;
+	}
+
+	static bool IsHLODActor(AActor* Actor)
+	{
+		return GetHLODObject(Actor) != nullptr;
+	}
+
 	static bool IsCollisionEnabled(UPrimitiveComponent* Component)
 	{
-		return Component->IsCollisionEnabled() && !Component->GetOwner()->IsA<AWorldPartitionHLOD>();
+		return Component->IsCollisionEnabled() && !IsHLODActor(Component->GetOwner());
 	}
 
 	static FString GetComponentShortName(UActorComponent* InComponent)
@@ -343,13 +362,22 @@ void UFastGeoWorldPartitionRuntimeCellTransformer::Transform(ULevel* InLevel)
 
 			FFastGeoComponentCluster* CurrentComponentCluster = LevelComponentCluster.Get();
 			TUniquePtr<FFastGeoHLOD> FastGeoHLOD;
-			if (AWorldPartitionHLOD* HLODActor = Cast<AWorldPartitionHLOD>(Actor))
+			IWorldPartitionHLODObject* HLODActor = FastGeo::GetHLODObject(Actor);
+			if (HLODActor)
 			{
 				FastGeoHLOD.Reset(new FFastGeoHLOD(FastGeo, *FString::Printf(TEXT("FastGeoHLOD_%s"), *Actor->GetName())));
 				FastGeoHLOD->SetSourceCellGuid(HLODActor->GetSourceCellGuid());
 				FastGeoHLOD->SetRequireWarmup(HLODActor->DoesRequireWarmup());
-				FastGeoHLOD->SetStandaloneHLODGuid(HLODActor->GetStandaloneHLODGuid());
-				FastGeoHLOD->SetCustomHLODGuid(HLODActor->GetCustomHLODGuid());
+
+				if (HLODActor->IsStandalone())
+				{
+					FastGeoHLOD->SetStandaloneHLODGuid(HLODActor->GetStandaloneHLODGuid());
+				}
+				else if (HLODActor->IsCustomHLOD())
+				{
+					FastGeoHLOD->SetCustomHLODGuid(HLODActor->GetCustomHLODGuid());
+				}
+
 				CurrentComponentCluster = FastGeoHLOD.Get();
 			}
 
