[FastGeo] Support users providing spatial hash info which defines bounds for ranges of instances and early culling on CPU in renderer.

Can significantly reduce memory allocated by renderer for draw calls on low end.

#rnx

#Synced-CL 50297198

diff --git a/Engine/Plugins/Experimental/FastGeoStreaming/Source/FastGeoStreaming/Internal/FastGeoProceduralISMComponent.h b/Engine/Plugins/Experimental/FastGeoStreaming/Source/FastGeoStreaming/Internal/FastGeoProceduralISMComponent.h
--- a/Engine/Plugins/Experimental/FastGeoStreaming/Source/FastGeoStreaming/Internal/FastGeoProceduralISMComponent.h
+++ b/Engine/Plugins/Experimental/FastGeoStreaming/Source/FastGeoStreaming/Internal/FastGeoProceduralISMComponent.h
@@ -20,6 +20,9 @@ public:
 
 	FASTGEOSTREAMING_API void InitializeFromComponentDescriptor(const FProceduralISMComponentDescriptor& InDescriptor);
 
+	/** Set spatial hashes which provide bounds for sub-ranges of instances within this primitive, enabling early CPU culling. */
+	FASTGEOSTREAMING_API void SetSpatialHashes(TArray<FInstanceSceneDataBuffers::FCompressedSpatialHashItem>&& InSpatialHashes);
+
 	//~ Begin FFastGeoComponent Interface
 	virtual void Serialize(FArchive& Ar) override;
 	//~ End FFastGeoComponent Interface
@@ -56,6 +59,7 @@ private:
 
 	// Transient data
 	TSharedPtr<FInstanceDataSceneProxy, ESPMode::ThreadSafe> DataProxy{};
+	TArray<FInstanceSceneDataBuffers::FCompressedSpatialHashItem> SpatialHashes;
 
 	friend class FInstancedStaticMeshComponentHelper;
 };
diff --git a/Engine/Plugins/Experimental/FastGeoStreaming/Source/FastGeoStreaming/Private/FastGeoProceduralISMComponent.cpp b/Engine/Plugins/Experimental/FastGeoStreaming/Source/FastGeoStreaming/Private/FastGeoProceduralISMComponent.cpp
--- a/Engine/Plugins/Experimental/FastGeoStreaming/Source/FastGeoStreaming/Private/FastGeoProceduralISMComponent.cpp
+++ b/Engine/Plugins/Experimental/FastGeoStreaming/Source/FastGeoStreaming/Private/FastGeoProceduralISMComponent.cpp
@@ -56,6 +56,15 @@ void FFastGeoProceduralISMComponent::InitializeFromComponentDescriptor(const FPr
 	SceneProxyDesc.bIsInstancedStaticMesh = true;
 }
 
+void FFastGeoProceduralISMComponent::SetSpatialHashes(TArray<FInstanceSceneDataBuffers::FCompressedSpatialHashItem>&& InSpatialHashes)
+{
+	if (SpatialHashes != InSpatialHashes)
+	{
+		SpatialHashes = MoveTemp(InSpatialHashes);
+		MarkRenderStateDirty();
+	}
+}
+
 void FFastGeoProceduralISMComponent::Serialize(FArchive& Ar)
 {
 	Super::Serialize(Ar);
@@ -128,6 +137,11 @@ TSharedPtr<FInstanceDataSceneProxy, ESPMode::ThreadSafe>& FFastGeoProceduralISMC
 	View.Flags.bHasPerInstanceCustomData = NumCustomDataFloats > 0;
 	View.Flags.bHasPerInstanceRandom = PrimitiveMaterialDesc.bAnyMaterialHasPerInstanceRandom;
 
+	if (!SpatialHashes.IsEmpty())
+	{
+		InstanceSceneDataBuffers.SetImmutable(FInstanceSceneDataImmutable(SpatialHashes), AccessTag);
+	}
+
 	InstanceSceneDataBuffers.EndWriteAccess(AccessTag);
 	InstanceSceneDataBuffers.ValidateData();
 
