[FastGeo] Enable GPU LOD selection by default.

This is typically enabled everywhere, but is false in the proxy by default.

#rnx

#Synced-CL 50393138

diff --git a/Engine/Plugins/Experimental/FastGeoStreaming/Source/FastGeoStreaming/Private/FastGeoProceduralISMComponent.cpp b/Engine/Plugins/Experimental/FastGeoStreaming/Source/FastGeoStreaming/Private/FastGeoProceduralISMComponent.cpp
--- a/Engine/Plugins/Experimental/FastGeoStreaming/Source/FastGeoStreaming/Private/FastGeoProceduralISMComponent.cpp
+++ b/Engine/Plugins/Experimental/FastGeoStreaming/Source/FastGeoStreaming/Private/FastGeoProceduralISMComponent.cpp
@@ -92,6 +92,7 @@ void FFastGeoProceduralISMComponent::ResetSceneProxyDescUnsupportedProperties()
 	SceneProxyDesc.bAffectDynamicIndirectLighting = false;
 	SceneProxyDesc.bAffectDistanceFieldLighting = false;
 	SceneProxyDesc.bCollisionEnabled = false;
+	SceneProxyDesc.bUseGpuLodSelection = true; // TODO: Remove this when the default value in the proxy desc is changed.
 }
 #endif
 
@@ -102,6 +103,7 @@ void FFastGeoProceduralISMComponent::InitializeSceneProxyDescDynamicProperties()
 	SceneProxyDesc.bAffectDynamicIndirectLighting = false;
 	SceneProxyDesc.bAffectDistanceFieldLighting = false;
 	SceneProxyDesc.bCollisionEnabled = false;
+	SceneProxyDesc.bUseGpuLodSelection = true; // TODO: Remove this when the default value in the proxy desc is changed.
 }
 
 void FFastGeoProceduralISMComponent::ApplyWorldTransform(const FTransform& InTransform)
