[PCG] Fix runtime gen GPU execution sources being left partially without instances when GPU scene instances are flushed during generation

#Synced-CL 50398172

diff --git a/Engine/Plugins/Experimental/FastGeoStreaming/Source/FastGeoStreaming/Internal/FastGeoAsyncRenderStateJobQueue.h b/Engine/Plugins/Experimental/FastGeoStreaming/Source/FastGeoStreaming/Internal/FastGeoAsyncRenderStateJobQueue.h
new file mode 100644
--- /dev/null
+++ b/Engine/Plugins/Experimental/FastGeoStreaming/Source/FastGeoStreaming/Internal/FastGeoAsyncRenderStateJobQueue.h
@@ -0,0 +1,92 @@
+// Copyright Epic Games, Inc. All Rights Reserved.
+
+#pragma once
+
+#include "Tasks/Task.h"
+#include "Tasks/Pipe.h"
+#include "FastGeoContainer.h"
+
+class FSceneInterface;
+
+struct FFastGeoAsyncRenderStateJobQueue
+{
+	FFastGeoAsyncRenderStateJobQueue(FSceneInterface* Scene);
+	~FFastGeoAsyncRenderStateJobQueue();
+
+	enum class EJobType
+	{
+		PrecachePSOs,
+		CreateRenderState,
+		DestroyRenderState
+	};
+
+	struct FJob
+	{
+		FJob(UFastGeoContainer* InFastGeo, EJobType InType)
+			: FastGeo(InFastGeo)
+			, Type(InType)
+		{}
+
+		bool IsValid() const { return ::IsValid(FastGeo); }
+		void Execute() const;
+		void OnPostExecute_GameThread() const;
+
+		inline bool operator == (const FJob& Other) const
+		{
+			return (FastGeo == Other.FastGeo) && (Type == Other.Type);
+		}
+
+		friend inline uint32 GetTypeHash(const FJob& InJob)
+		{
+			return GetTypeHash(TPair<UFastGeoContainer*, EJobType>{InJob.FastGeo, InJob.Type});
+		}
+
+		UFastGeoContainer* FastGeo;
+		EJobType Type;
+	};
+
+	struct FJobSet
+	{
+		void Add(const FJob& Job)
+		{
+			Jobs.Add(Job);
+		}
+
+		int32 Remove(const FJob& Job)
+		{
+			return Jobs.Remove(Job);
+		}
+
+		bool IsEmpty() const
+		{
+			return Jobs.IsEmpty();
+		}
+
+		void Execute() const;
+		void OnPostExecute_GameThread() const;
+
+	private:
+		TSet<FJob> Jobs;
+	};
+
+	void Tick(bool bForceWaitCompletion);
+	void AddJob(const FJob& Job);
+	bool IsCompleted() const;
+
+private:
+
+	void Launch();
+	void WaitForAsyncTasksExecution();
+	bool AreAsyncTasksExecuted() const;
+	void OnAsyncTasksExecuted();
+	void OnUpdateLevelStreamingDone();
+	void TriggerIsReadyToRunAsyncTasksEvent();
+
+	FSceneInterface* Scene;
+	using FJobs = TUniquePtr<FJobSet>;
+	FJobs PendingJobs;
+	TArray<FJobs> PipedJobs;
+	TArray<UE::Tasks::FTask> PipeTasks;
+	TOptional<UE::Tasks::FTaskEvent> IsReadyToRunAsyncTasksEvent;
+	UE::Tasks::FPipe Pipe{ TEXT("FastGeoAsyncJobQueue") };
+};
\ No newline at end of file
diff --git a/Engine/Plugins/Experimental/FastGeoStreaming/Source/FastGeoStreaming/Internal/FastGeoRegisteredComponent.h b/Engine/Plugins/Experimental/FastGeoStreaming/Source/FastGeoStreaming/Internal/FastGeoRegisteredComponent.h
new file mode 100644
--- /dev/null
+++ b/Engine/Plugins/Experimental/FastGeoStreaming/Source/FastGeoStreaming/Internal/FastGeoRegisteredComponent.h
@@ -0,0 +1,41 @@
+// Copyright Epic Games, Inc. All Rights Reserved.
+
+#pragma once
+
+#include "FastGeoComponent.h"
+#include "FastGeoContainer.h"
+#include "FastGeoWeakElement.h"
+
+struct FFastGeoRegisteredComponent
+{
+	FFastGeoRegisteredComponent(FFastGeoComponent* InComponent)
+		: Component(InComponent)
+	{
+		check(InComponent);
+		check(InComponent->GetOwnerContainer());
+		RegistrationEpoch = InComponent->GetOwnerContainer()->RegistrationEpoch;
+	}
+
+	template <class T>
+	FORCEINLINE T* TryGetRegistered() const
+	{
+		if (T* ComponentPtr = Component.Get<T>())
+		{
+			if (ComponentPtr->GetOwnerContainer()->IsRegistered() &&
+				RegistrationEpoch == ComponentPtr->GetOwnerContainer()->RegistrationEpoch)
+			{
+				return ComponentPtr;
+			}
+		}
+		return nullptr;
+	}
+
+	FORCEINLINE bool operator==(const FFastGeoRegisteredComponent& Other) const
+	{
+		return Component == Other.Component && RegistrationEpoch == Other.RegistrationEpoch;
+	}
+
+private:
+	FWeakFastGeoComponent Component;
+	int32 RegistrationEpoch = -1;
+};
\ No newline at end of file
diff --git a/Engine/Plugins/Experimental/FastGeoStreaming/Source/FastGeoStreaming/Internal/FastGeoWorldSubsystem.h b/Engine/Plugins/Experimental/FastGeoStreaming/Source/FastGeoStreaming/Internal/FastGeoWorldSubsystem.h
new file mode 100644
--- /dev/null
+++ b/Engine/Plugins/Experimental/FastGeoStreaming/Source/FastGeoStreaming/Internal/FastGeoWorldSubsystem.h
@@ -0,0 +1,145 @@
+// Copyright Epic Games, Inc. All Rights Reserved.
+
+#pragma once
+
+#include "Subsystems/WorldSubsystem.h"
+#include "FastGeoWeakElement.h"
+#include "FastGeoRegisteredComponent.h"
+#include "FastGeoAsyncRenderStateJobQueue.h"
+#include "FastGeoWorldSubsystem.generated.h"
+
+class FFastGeoPrimitiveComponent;
+class ULevelStreaming;
+class ULevel;
+class IWorldPartitionHLODObject;
+enum class ELevelStreamingState : uint8;
+
+UCLASS()
+class FASTGEOSTREAMING_API UFastGeoWorldSubsystem : public UTickableWorldSubsystem
+{
+	GENERATED_BODY()
+
+	UFastGeoWorldSubsystem();
+	~UFastGeoWorldSubsystem();
+
+public:
+	//~Begin USubsystem interface
+	virtual void Deinitialize() override;
+	//~End USubsystem interface
+
+	//~ Begin UWorldSubsystem interface.
+	virtual void PostInitialize() override;
+	virtual bool DoesSupportWorldType(const EWorldType::Type WorldType) const override;
+	//~ End UWorldSubsystem interface.
+
+	//~ Begin UTickableWorldSubsystem interface
+	virtual bool IsTickableInEditor() const override;
+	virtual void Tick(float DeltaTime) override;
+	virtual TStatId GetStatId() const override;
+	//~ End UTickableWorldSubsystem interface
+
+	void AddToComponentsPendingRecreate(FFastGeoComponent* InComponentPendingRecreate);
+	void ProcessPendingRecreate();
+
+	/**
+	 * Start a job that will precache PSOs asynchronously for a specified FastGeo container.
+	 * @param FastGeo: FastGeo container.
+	 */
+	void PushAsyncPrecachePSOsJob(UFastGeoContainer* FastGeo);
+
+	/**
+	 * Start a job that will create render state asynchronously for a specified FastGeo container.
+	 * @param FastGeo: FastGeo container.
+	 */
+	void PushAsyncCreateRenderStateJob(UFastGeoContainer* FastGeo);
+
+	/**
+	 * Start a job that will destroy render state asynchronously for a specified FastGeo container.
+	 * @param FastGeo: FastGeo container.
+	 */
+	void PushAsyncDestroyRenderStateJob(UFastGeoContainer* FastGeo);
+
+	/**
+	 * Update progress of asynchronous render state creation and destruction.
+	 * @param bWaitForCompletion : Whether to wait for pending asynchronous tasks to complete.
+	 */
+	void ProcessAsyncRenderStateJobs(bool bWaitForCompletion = false);
+
+	/**
+	 * Start a job that will create physics state asynchronously for a specified FastGeo container.
+	 * @param FastGeo: FastGeo container.
+	 */
+	void PushAsyncCreatePhysicsStateJobs(UFastGeoContainer* FastGeo);
+
+	/**
+	 * Start a job that will destroy physics state asynchronously for a specified FastGeo container.
+	 * @param FastGeo: FastGeo container.
+	 */
+	void PushAsyncDestroyPhysicsStateJobs(UFastGeoContainer* FastGeo);
+
+	static bool IsEnableDebugView();
+
+#if !UE_BUILD_SHIPPING
+	static bool IsFastGeoVisible();
+#endif
+
+	bool IsWaitingForCompletion() const;
+
+#if WITH_EDITOR
+public:
+	DECLARE_MULTICAST_DELEGATE_OneParam(FOnComponentsPreRecreate, const TArray<FFastGeoRegisteredComponent>&);
+
+	/** Called right before processing any registered FastGeo components pending recreate each frame. */
+	static FOnComponentsPreRecreate ComponentsPreRecreateEvent;
+#endif
+
+private:
+	void OnUpdateLevelStreaming();
+	void OnWorldCleanup(UWorld* World, bool bSessionEnded, bool bCleanupResources);
+	void OnLevelStreamingStateChanged(UWorld* World, const ULevelStreaming* StreamingLevel, ULevel* LevelIfLoaded, ELevelStreamingState PrevState, ELevelStreamingState NewState);
+	void OnLevelStartedAddToWorld(UWorld* World, const ULevelStreaming* StreamingLevel, ULevel* Level);
+	void OnLevelStartedRemoveFromWorld(UWorld* World, const ULevelStreaming* StreamingLevel, ULevel* Level);
+	void OnLevelComponentsUpdated(UWorld* World, ULevel* Level);
+	void OnLevelComponentsCleared(UWorld* World, ULevel* Level);
+	void OnAddLevelToWorldExtension(ULevel* Level, const bool bWaitForCompletion, bool& bOutHasCompleted);
+	void OnRemoveLevelFromWorldExtension(ULevel* Level, const bool bWaitForCompletion, bool& bOutHasCompleted);
+	void OnGlobalComponentReregisterContextCreated();
+	void OnGlobalComponentReregisterContextDestroyed();
+#if DO_CHECK
+	void OnLevelAddedToWorld(ULevel* Level, UWorld* World);
+	void OnLevelRemovedFromWorld(ULevel* Level, UWorld* World);
+	void CheckNoPendingTasks(ULevel* Level, UWorld* World);
+#endif
+#if WITH_EDITOR
+	void OnPreRecreateScene(UWorld* World);
+#endif
+	void ForEachHLODObjectInCell(const UWorldPartitionRuntimeCell*, TFunction<void(IWorldPartitionHLODObject*)>);
+	void RequestAsyncRenderStateTasksBudget_Concurrent(float& OutAvailableTimeBudgetMS, int32& OutAvaiableComponentsBudget, int32& OutTimeEpoch);
+	void CommitAsyncRenderStateTasksBudget_Concurrent(float InUsedTimeBudgetMS, int32& InUsedComponentsBudget, int32 TimeEpoch);
+	void PushAsyncDestroyContainer(UFastGeoContainer* FastGeo);
+	FFastGeoAsyncRenderStateJobQueue* GetOrCreateJobQueue();
+	void WaitForAllAsyncJobsCompletion();
+
+	UPROPERTY()
+	TSet<TObjectPtr<UFastGeoContainer>> PendingDestroyContainers;
+
+	FDelegateHandle Handle_OnLevelStreamingStateChanged;
+	FDelegateHandle Handle_OnLevelBeginAddToWorld;
+	FDelegateHandle Handle_OnLevelBeginRemoveFromWorld;
+	FDelegateHandle Handle_OnForEachHLODObjectInCell;
+
+	TArray<FFastGeoRegisteredComponent> ComponentsPendingRecreate;
+	static bool bEnableDebugView;
+	static bool bIsFastGeoVisible;
+
+	FRWLock Lock;
+	int32 TimeEpoch = 0;
+	float UsedAsyncRenderStateTasksTimeBudgetMS = 0;
+	int32 UsedNumComponentsToProcessBudget = 0;
+	bool bWaitingForCompletion = false;
+
+	TUniquePtr<FFastGeoAsyncRenderStateJobQueue> AsyncRenderStateJobQueue;
+	TArray<TWeakObjectPtr<UFastGeoContainer>> ReregisteringFastGeoContainers;
+
+	friend class UFastGeoContainer;
+};
diff --git a/Engine/Plugins/Experimental/FastGeoStreaming/Source/FastGeoStreaming/Private/FastGeoAsyncRenderStateJobQueue.h b/Engine/Plugins/Experimental/FastGeoStreaming/Source/FastGeoStreaming/Private/FastGeoAsyncRenderStateJobQueue.h
deleted file mode 100644
diff --git a/Engine/Plugins/Experimental/FastGeoStreaming/Source/FastGeoStreaming/Private/FastGeoRegisteredComponent.h b/Engine/Plugins/Experimental/FastGeoStreaming/Source/FastGeoStreaming/Private/FastGeoRegisteredComponent.h
deleted file mode 100644
diff --git a/Engine/Plugins/Experimental/FastGeoStreaming/Source/FastGeoStreaming/Private/FastGeoWorldSubsystem.cpp b/Engine/Plugins/Experimental/FastGeoStreaming/Source/FastGeoStreaming/Private/FastGeoWorldSubsystem.cpp
--- a/Engine/Plugins/Experimental/FastGeoStreaming/Source/FastGeoStreaming/Private/FastGeoWorldSubsystem.cpp
+++ b/Engine/Plugins/Experimental/FastGeoStreaming/Source/FastGeoStreaming/Private/FastGeoWorldSubsystem.cpp
@@ -22,6 +22,10 @@
 static FName NAME_FastGeoColorHandler(TEXT("FastGeo"));
 bool UFastGeoWorldSubsystem::bEnableDebugView = false;
 
+#if WITH_EDITOR
+UFastGeoWorldSubsystem::FOnComponentsPreRecreate UFastGeoWorldSubsystem::ComponentsPreRecreateEvent;
+#endif
+
 namespace FastGeo
 {
 	static float GAsyncRenderStateTaskTimeBudgetMS = 0.0f;
@@ -467,6 +471,13 @@ void UFastGeoWorldSubsystem::ProcessPendingRecreate()
 
 	CSV_CUSTOM_STAT(FastGeo, PendingRecreate, ComponentsPendingRecreate.Num(), ECsvCustomStatOp::Set);
 
+#if WITH_EDITOR
+	if (!ComponentsPendingRecreate.IsEmpty())
+	{
+		ComponentsPreRecreateEvent.Broadcast(ComponentsPendingRecreate);
+	}
+#endif
+
 	int32 NumPrimitiveDelayed = 0;
 	for (const FFastGeoRegisteredComponent& PendingRecreateComponent : ComponentsPendingRecreate)
 	{
diff --git a/Engine/Plugins/Experimental/FastGeoStreaming/Source/FastGeoStreaming/Private/FastGeoWorldSubsystem.h b/Engine/Plugins/Experimental/FastGeoStreaming/Source/FastGeoStreaming/Private/FastGeoWorldSubsystem.h
deleted file mode 100644
diff --git a/Engine/Plugins/Experimental/PCGInterops/PCGFastGeoInterop/Source/PCGFastGeoInterop/Private/PCGFastGeoInteropModule.cpp b/Engine/Plugins/Experimental/PCGInterops/PCGFastGeoInterop/Source/PCGFastGeoInterop/Private/PCGFastGeoInteropModule.cpp
--- a/Engine/Plugins/Experimental/PCGInterops/PCGFastGeoInterop/Source/PCGFastGeoInterop/Private/PCGFastGeoInteropModule.cpp
+++ b/Engine/Plugins/Experimental/PCGInterops/PCGFastGeoInterop/Source/PCGFastGeoInterop/Private/PCGFastGeoInteropModule.cpp
@@ -4,6 +4,10 @@
 
 #include "Compute/PrimitiveFactories/PCGPrimitiveFactoryFastGeoPISMC.h"
 
+#include "PCGSubsystem.h"
+
+#include "FastGeoWorldSubsystem.h"
+
 #include "Modules/ModuleManager.h"
 
 IMPLEMENT_MODULE(FPCGFastGeoInteropModule, PCGFastGeoInterop);
@@ -14,9 +18,28 @@ void FPCGFastGeoInteropModule::StartupModule()
 	{
 		return MakeShared<FPCGPrimitiveFactoryFastGeoPISMC>();
 	});
+
+#if WITH_EDITOR
+	UFastGeoWorldSubsystem::ComponentsPreRecreateEvent.AddRaw(this, &FPCGFastGeoInteropModule::OnFastGeoComponentsPreRecreate);
+#endif
 }
 
 void FPCGFastGeoInteropModule::ShutdownModule()
 {
+#if WITH_EDITOR
+	UFastGeoWorldSubsystem::ComponentsPreRecreateEvent.RemoveAll(this);
+#endif
+
 	PCGPrimitiveFactoryHelpers::Private::ResetFastGeoPrimitiveFactory();
 }
+
+#if WITH_EDITOR
+void FPCGFastGeoInteropModule::OnFastGeoComponentsPreRecreate(const TArray<FFastGeoRegisteredComponent>& InComponents) const
+{
+	if (UPCGSubsystem* PCGSubsystem = UPCGSubsystem::GetSubsystemForCurrentWorld())
+	{
+		// Could try to be more clever and only refresh the component that was dirtied, but if the instances were flushed we are probably safe to refresh everything.
+		PCGSubsystem->DirtyRuntimeGenExecutionSources();
+	}
+}
+#endif
diff --git a/Engine/Plugins/Experimental/PCGInterops/PCGFastGeoInterop/Source/PCGFastGeoInterop/Public/PCGFastGeoInteropModule.h b/Engine/Plugins/Experimental/PCGInterops/PCGFastGeoInterop/Source/PCGFastGeoInterop/Public/PCGFastGeoInteropModule.h
--- a/Engine/Plugins/Experimental/PCGInterops/PCGFastGeoInterop/Source/PCGFastGeoInterop/Public/PCGFastGeoInteropModule.h
+++ b/Engine/Plugins/Experimental/PCGInterops/PCGFastGeoInterop/Source/PCGFastGeoInterop/Public/PCGFastGeoInteropModule.h
@@ -14,4 +14,8 @@ public:
 	virtual void ShutdownModule() override;
 	virtual bool SupportsDynamicReloading() override { return true; }
 	//~ End IModuleInterface implementation
+
+#if WITH_EDITOR
+	void OnFastGeoComponentsPreRecreate(const TArray<struct FFastGeoRegisteredComponent>& InComponents) const;
+#endif
 };
diff --git a/Engine/Plugins/PCG/Source/PCG/Private/Components/PCGProceduralISMComponent.cpp b/Engine/Plugins/PCG/Source/PCG/Private/Components/PCGProceduralISMComponent.cpp
--- a/Engine/Plugins/PCG/Source/PCG/Private/Components/PCGProceduralISMComponent.cpp
+++ b/Engine/Plugins/PCG/Source/PCG/Private/Components/PCGProceduralISMComponent.cpp
@@ -407,7 +407,6 @@ void UPCGProceduralISMComponent::ValidateComponentSetup()
 void UPCGProceduralISMComponent::OnRenderStateDirty(UActorComponent& InComponent)
 {
 	// Currently, there is no explicit persistence of instance data in the GPU scene. When this component is dirtied, the instance data is cleared.
-	// TODO: This function is a stop gap that requests a refresh of the PCG Component managing this component, and should be removed later.
 
 	if (&InComponent != this)
 	{
@@ -415,62 +414,11 @@ void UPCGProceduralISMComponent::OnRenderStateDirty(UActorComponent& InComponent
 	}
 
 	UPCGSubsystem* Subsystem = GetOwner() ? UPCGSubsystem::GetInstance(GetOwner()->GetWorld()) : nullptr;
-	if (!Subsystem)
-	{
-		return;
-	}
 
-	// Helper that returns true if the given PCG component is managing this PISMC.
-	auto PCGComponentManagesThisPISMC = [this](UPCGComponent* InComponent) -> bool
+	if (Subsystem)
 	{
-		bool bManagesThis = false;
-
-		if (InComponent->bGenerated && InComponent->AreProceduralInstancesInUse() && InComponent->AreManagedResourcesAccessible())
-		{
-			InComponent->ForEachManagedResource([this, InComponent, &bManagesThis](UPCGManagedResource* InResource)
-			{
-				if (!bManagesThis)
-				{
-					const UPCGManagedProceduralISMComponent* PISMC = Cast<UPCGManagedProceduralISMComponent>(InResource);
-					if (PISMC && PISMC->GetComponent() == this)
-					{
-						bManagesThis = true;
-					}
-				}
-			});
-		}
-
-		return bManagesThis;
-	};
-
-	Subsystem->RefreshAllComponentsFiltered(
-		[this, Subsystem, &PCGComponentManagesThisPISMC](UPCGComponent* InComponent) -> bool
-		{
-			// If the original component manages this PISMC, request a refresh of it and we're done. If it has local components they will also be refreshed.
-			if (PCGComponentManagesThisPISMC(InComponent))
-			{
-				return true;
-			}
-
-			// A local component of the original component might manage this, so check those.
-			bool bManagesThis = false;
-
-			if (InComponent->bIsComponentPartitioned)
-			{
-				Subsystem->ForAllRegisteredLocalComponents(
-					InComponent,
-					[this, &bManagesThis, &PCGComponentManagesThisPISMC](UPCGComponent* InComponent)
-					{
-						if (!bManagesThis)
-						{
-							bManagesThis = PCGComponentManagesThisPISMC(InComponent);
-						}
-					});
-			}
-
-			return bManagesThis;
-		},
-		EPCGChangeType::Structural);
+		Subsystem->DirtyRuntimeGenExecutionSources();
+	}
 }
 #endif // WITH_EDITOR
 
diff --git a/Engine/Plugins/PCG/Source/PCG/Private/Subsystems/PCGSubsystem.cpp b/Engine/Plugins/PCG/Source/PCG/Private/Subsystems/PCGSubsystem.cpp
--- a/Engine/Plugins/PCG/Source/PCG/Private/Subsystems/PCGSubsystem.cpp
+++ b/Engine/Plugins/PCG/Source/PCG/Private/Subsystems/PCGSubsystem.cpp
@@ -367,6 +367,17 @@ void UPCGSubsystem::Tick(float DeltaSeconds)
 		return;
 	}
 
+#if WITH_EDITOR
+	// Refresh dirty runtime gen sources before ticking the graph executor (in IPCGBaseSubsystem::Tick()) so that we're not processing a bunch of graphs/execution sources that will be cancelled immediately.
+	if (bDirtyRuntimeGenExecutionSources)
+	{
+		RefreshAllRuntimeGenExecutionSources(DirtyRuntimeGenChangeType);
+
+		bDirtyRuntimeGenExecutionSources = false;
+		DirtyRuntimeGenChangeType = EPCGChangeType::None;
+	}
+#endif
+
 	Super::Tick(DeltaSeconds);
 
 	ExecuteBeginTickActions();
@@ -1065,6 +1076,12 @@ void UPCGSubsystem::RefreshAllRuntimeGenExecutionSources(EPCGChangeType ChangeTy
 }
 
 #if WITH_EDITOR
+void UPCGSubsystem::DirtyRuntimeGenExecutionSources(EPCGChangeType ChangeType)
+{
+	bDirtyRuntimeGenExecutionSources = true;
+	DirtyRuntimeGenChangeType |= ChangeType;
+}
+
 void UPCGSubsystem::RefreshAllComponentsFiltered(const TFunction<bool(UPCGComponent*)>& ComponentFilter, EPCGChangeType ChangeType)
 {
 	// @todo_pcg: support execution source
diff --git a/Engine/Plugins/PCG/Source/PCG/Public/Subsystems/PCGSubsystem.h b/Engine/Plugins/PCG/Source/PCG/Public/Subsystems/PCGSubsystem.h
--- a/Engine/Plugins/PCG/Source/PCG/Public/Subsystems/PCGSubsystem.h
+++ b/Engine/Plugins/PCG/Source/PCG/Public/Subsystems/PCGSubsystem.h
@@ -207,13 +207,16 @@ public:
 	*/
 	PCG_API FPCGTaskId ScheduleGenericWithContext(TFunction<bool(FPCGContext*)> InOperation, TFunction<void(FPCGContext*)> InAbortOperation, IPCGGraphExecutionSource* ExecutionSource, const TArray<FPCGTaskId>& TaskExecutionDependencies, const TArray<FPCGTaskId>& TaskDataDependencies) const;
 
-	/** Asks the runtime generation scheduler to refresh a given GenerateAtRuntime execution source. ChangeType should be 'GenerationGrid' to perform a full cleanup of PAs and local components. */
+	/** Asks the runtime generation scheduler to refresh a given GenerateAtRuntime execution source. ChangeType should be 'GenerationGrid' to perform a full cleanup of PAs and local sources. */
 	PCG_API void RefreshRuntimeGenExecutionSource(IPCGGraphExecutionSource* RuntimeExecutionSource, EPCGChangeType ChangeType = EPCGChangeType::None);
 
-	/** Asks the runtime generation scheduler to refresh all GenerateAtRuntime execution sources. ChangeType should be 'GenerationGrid' to perform a full cleanup of PAs and local components. */
+	/** Asks the runtime generation scheduler to refresh all GenerateAtRuntime execution sources. ChangeType should be 'GenerationGrid' to perform a full cleanup of PAs and local sources. */
 	PCG_API void RefreshAllRuntimeGenExecutionSources(EPCGChangeType ChangeType = EPCGChangeType::None);
 
 #if WITH_EDITOR
+	/** Asks the runtime generation scheduler to refresh all GenerateAtRuntime execution sources. ChangeType should be 'GenerationGrid' to perform a full cleanup of PAs and local sources. */
+	PCG_API void DirtyRuntimeGenExecutionSources(EPCGChangeType ChangeType = EPCGChangeType::None);
+
 	/** Refresh all components selected by the filter (runtime generated or otherwise). */
 	PCG_API void RefreshAllComponentsFiltered(const TFunction<bool(UPCGComponent*)>& ComponentFilter, EPCGChangeType ChangeType = EPCGChangeType::None);
 #endif
@@ -454,6 +457,9 @@ private:
 	using FConstructionScriptSourceComponents = TMap<FName, TObjectKey<UPCGComponent>>;
 	TMap<TObjectKey<AActor>, FConstructionScriptSourceComponents> PerActorConstructionScriptSourceComponents;
 
+	bool bDirtyRuntimeGenExecutionSources = false;
+	EPCGChangeType DirtyRuntimeGenChangeType = EPCGChangeType::None;
+
 	PCG_API static TSet<UWorld*> DisablePartitionActorCreationForWorld;
 
 	// Used by UPCGWorldPartitonBuilder to disable PA creation while outside of a certain scope
diff --git a/Engine/Plugins/PCG/Source/PCGEditor/Private/PCGEditorModule.cpp b/Engine/Plugins/PCG/Source/PCGEditor/Private/PCGEditorModule.cpp
--- a/Engine/Plugins/PCG/Source/PCGEditor/Private/PCGEditorModule.cpp
+++ b/Engine/Plugins/PCG/Source/PCGEditor/Private/PCGEditorModule.cpp
@@ -343,12 +343,9 @@ void FPCGEditorModule::OnSceneMaterialsModified()
 	}
 
 	// Currently, there is no explicit persistence of instance data in the GPU scene and procedural instances are lost when the GPU Scene is flushed.
-	// TODO: This function is a stop gap that refreshes PCG Components owning procedural instances, and should be removed later.
 	if (UPCGSubsystem* Subsystem = UPCGSubsystem::GetSubsystemForCurrentWorld())
 	{
-		Subsystem->RefreshAllComponentsFiltered(
-			[](UPCGComponent* InComponent) { return InComponent->AreProceduralInstancesInUse(); },
-			EPCGChangeType::Structural);
+		Subsystem->DirtyRuntimeGenExecutionSources();
 	}
 }
 
