Enable GPU LOD selection by default in ISM proxy desc.

#rnx

#Synced-CL 50412500

diff --git a/Engine/Plugins/Experimental/FastGeoStreaming/Source/FastGeoStreaming/Private/FastGeoProceduralISMComponent.cpp b/Engine/Plugins/Experimental/FastGeoStreaming/Source/FastGeoStreaming/Private/FastGeoProceduralISMComponent.cpp
--- a/Engine/Plugins/Experimental/FastGeoStreaming/Source/FastGeoStreaming/Private/FastGeoProceduralISMComponent.cpp
+++ b/Engine/Plugins/Experimental/FastGeoStreaming/Source/FastGeoStreaming/Private/FastGeoProceduralISMComponent.cpp
@@ -92,7 +92,6 @@ void FFastGeoProceduralISMComponent::ResetSceneProxyDescUnsupportedProperties()
 	SceneProxyDesc.bAffectDynamicIndirectLighting = false;
 	SceneProxyDesc.bAffectDistanceFieldLighting = false;
 	SceneProxyDesc.bCollisionEnabled = false;
-	SceneProxyDesc.bUseGpuLodSelection = true; // TODO: Remove this when the default value in the proxy desc is changed.
 }
 #endif
 
@@ -103,7 +102,6 @@ void FFastGeoProceduralISMComponent::InitializeSceneProxyDescDynamicProperties()
 	SceneProxyDesc.bAffectDynamicIndirectLighting = false;
 	SceneProxyDesc.bAffectDistanceFieldLighting = false;
 	SceneProxyDesc.bCollisionEnabled = false;
-	SceneProxyDesc.bUseGpuLodSelection = true; // TODO: Remove this when the default value in the proxy desc is changed.
 }
 
 void FFastGeoProceduralISMComponent::ApplyWorldTransform(const FTransform& InTransform)
diff --git a/Engine/Source/Runtime/Engine/Public/InstancedSkinnedMeshSceneProxyDesc.h b/Engine/Source/Runtime/Engine/Public/InstancedSkinnedMeshSceneProxyDesc.h
--- a/Engine/Source/Runtime/Engine/Public/InstancedSkinnedMeshSceneProxyDesc.h
+++ b/Engine/Source/Runtime/Engine/Public/InstancedSkinnedMeshSceneProxyDesc.h
@@ -40,5 +40,5 @@ struct FInstancedSkinnedMeshSceneProxyDesc : public FSkinnedMeshSceneProxyDesc
 	TBitArray<> SelectedInstances;
 #endif
 
-	bool bUseGpuLodSelection = false;
+	bool bUseGpuLodSelection = true;
 };
diff --git a/Engine/Source/Runtime/Engine/Public/InstancedStaticMeshSceneProxyDesc.h b/Engine/Source/Runtime/Engine/Public/InstancedStaticMeshSceneProxyDesc.h
--- a/Engine/Source/Runtime/Engine/Public/InstancedStaticMeshSceneProxyDesc.h
+++ b/Engine/Source/Runtime/Engine/Public/InstancedStaticMeshSceneProxyDesc.h
@@ -27,6 +27,6 @@ struct FInstancedStaticMeshSceneProxyDesc : public FStaticMeshSceneProxyDesc
 	int32 InstanceEndCullDistance = 0;
 	float InstanceLODDistanceScale = 1.0f;
 
-	bool bUseGpuLodSelection = false;
+	bool bUseGpuLodSelection = true;
 
 };
