[FastGeo] Added option to generate a surrogate Actor/Components per streaming cell (New FastGeo Transforms flag bGenerateSurrogateComponents, enabled by default).
- A FastGeoSurrogateComponent will be created to represent transformed FastGeo components that have the same specs (CollisionResponse, CollisionEnabled, ObjectType, WalkableSlopeOverride, CanCharacterStepUpOn).
- FastGeoSurrogateComponents don't have any render/physics state.
- This allows for actor/components hitting FastGeo content to trigger OnComponentHit.
- This fixes CitySample's damage system when cars hit FastGeo content (UE-305776).
- This fixes issues with code relying on a valid component and/or WalkableSlopeOverride info.
- When using FastGeoSurrogatecomponents, FHitResult::Item, FOverlapResult::ItemIndex and FCollisionNotifyInfo::BodyIndex will be automatically mapped to the actual FastGeo body instance index.
- FastGeoSurrogateComponent::GetBodyInstance will resolve to the actual FastGeoPrimitiveComponent's body instance when BodyInstanceIndex is passed to the function.
- Added CVar FastGeo.AllowSurrogateComponents to allow or not the use of generated FastGeo surrogate components (enabled by default).

#Synced-CL 50298075

diff --git a/Engine/Plugins/Experimental/FastGeoStreaming/Source/FastGeoStreaming/Internal/FastGeoComponent.h b/Engine/Plugins/Experimental/FastGeoStreaming/Source/FastGeoStreaming/Internal/FastGeoComponent.h
--- a/Engine/Plugins/Experimental/FastGeoStreaming/Source/FastGeoStreaming/Internal/FastGeoComponent.h
+++ b/Engine/Plugins/Experimental/FastGeoStreaming/Source/FastGeoStreaming/Internal/FastGeoComponent.h
@@ -22,6 +22,10 @@ class UFastGeoContainer;
 class FSceneInterface;
 class FPrimitiveSceneProxy;
 class FRegisterComponentContext;
+struct FBodyInstance;
+
+// Uncomment to debug FastGeo physics state creation/destruction
+// #define FASTGEO_DEBUG_PHYSICS
 
 enum class ESceneProxyCreationError
 {
@@ -59,11 +63,19 @@ public:
 	virtual void CreateRenderState(FRegisterComponentContext* Context) {}
 	virtual void DestroyRenderState(FFastGeoDestroyRenderStateContext* Context) {}
 	virtual void UpdateVisibility() {}
+#ifdef FASTGEO_DEBUG_PHYSICS
 	virtual void OnAsyncCreatePhysicsState();
 	virtual void OnAsyncDestroyPhysicsStateBegin_GameThread();
 	virtual void OnAsyncDestroyPhysicsState();
 	virtual void OnAsyncDestroyPhysicsStateEnd_GameThread();
 	void OnAsyncCreatePhysicsStateEnd_GameThread();
+#else
+	virtual void OnAsyncCreatePhysicsState() {}
+	virtual void OnAsyncDestroyPhysicsStateBegin_GameThread() {}
+	virtual void OnAsyncDestroyPhysicsState() {}
+	virtual void OnAsyncDestroyPhysicsStateEnd_GameThread() {}
+#endif
+	virtual void GetBodyInstances(TArray<FBodyInstance*>& OutBodyInstances) {}
 
 #if !WITH_EDITOR && UE_WITH_PSO_PRECACHING
 	virtual bool ShouldPrecachePSOs() const;
@@ -90,6 +102,7 @@ protected:
 	FFastGeoComponentCluster* Owner = nullptr;
 	TUniquePtr<FRWLock> Lock;
 
+#ifdef FASTGEO_DEBUG_PHYSICS
 	enum EPhysicsStateCreation
 	{
 		NotCreated,
@@ -99,6 +112,7 @@ protected:
 	};
 
 	EPhysicsStateCreation PhysicsStateCreation = EPhysicsStateCreation::NotCreated;
+#endif
 
 	enum class EProxyCreationState : uint8
 	{
diff --git a/Engine/Plugins/Experimental/FastGeoStreaming/Source/FastGeoStreaming/Internal/FastGeoContainer.h b/Engine/Plugins/Experimental/FastGeoStreaming/Source/FastGeoStreaming/Internal/FastGeoContainer.h
--- a/Engine/Plugins/Experimental/FastGeoStreaming/Source/FastGeoStreaming/Internal/FastGeoContainer.h
+++ b/Engine/Plugins/Experimental/FastGeoStreaming/Source/FastGeoStreaming/Internal/FastGeoContainer.h
@@ -26,6 +26,7 @@ class UNavigationSystemBase;
 class FFastGeoComponent;
 class FFastGeoComponentCluster;
 class FChaosUserDefinedEntity;
+class AFastGeoSurrogateActor;
 
 struct FFastGeoComponentClusterPendingState
 {
@@ -143,6 +144,9 @@ public:
 	virtual IPhysicsBodyInstanceOwner* ResolvePhysicsBodyInstanceOwner(Chaos::FConstPhysicsObjectHandle PhysicsObject) override;
 	//~End IPhysicsBodyInstanceOwnerResolver
 
+	bool IsUsingSurrogateComponents() const;
+	AFastGeoSurrogateActor* GetSurrogateActor() const;
+
 	void OnCreated(bool bCollectReferences = true);
 	void AddComponentCluster(FFastGeoComponentCluster* ComponentCluster);
 
@@ -214,6 +218,8 @@ private:
 	void RegisterToNavigationSystem();
 	void UnregisterFromNavigationSystem();
 	void OnNavigationInitDone(const UNavigationSystemBase& NavSys);
+	void RegisterBodyInstances();
+	void UnregisterBodyInstances();
 
 	template <typename ThisType, typename TComponentCluster, typename TFunc>
 	static void ForEachComponentCluster(ThisType* Self, TFunc&& InFunc)
@@ -297,6 +303,9 @@ private:
 	UPROPERTY()
 	TArray<TObjectPtr<UObject>> Assets;
 
+	UPROPERTY()
+	TObjectPtr<AFastGeoSurrogateActor> SurrogateActor;
+
 #if !WITH_EDITOR && UE_WITH_PSO_PRECACHING
 	bool bPrecachedPSOs = false;
 #endif
@@ -315,6 +324,7 @@ private:
 
 	FOnUnregistered OnUnregistered;
 
+	friend class UFastGeoWorldPartitionRuntimeCellTransformer;
 	friend class FFastGeoGatherFastGeoContainerAssetRefsArchive;
 	friend struct FFastGeoRegisteredComponent;
 	friend class UFastGeoWorldSubsystem;
diff --git a/Engine/Plugins/Experimental/FastGeoStreaming/Source/FastGeoStreaming/Internal/FastGeoInstancedStaticMeshComponent.h b/Engine/Plugins/Experimental/FastGeoStreaming/Source/FastGeoStreaming/Internal/FastGeoInstancedStaticMeshComponent.h
--- a/Engine/Plugins/Experimental/FastGeoStreaming/Source/FastGeoStreaming/Internal/FastGeoInstancedStaticMeshComponent.h
+++ b/Engine/Plugins/Experimental/FastGeoStreaming/Source/FastGeoStreaming/Internal/FastGeoInstancedStaticMeshComponent.h
@@ -36,6 +36,9 @@ public:
 	virtual Chaos::FPhysicsObject* GetPhysicsObjectById(Chaos::FPhysicsObjectId Id) const override;
 	virtual TArray<Chaos::FPhysicsObject*> GetAllPhysicsObjects() const override;
 	virtual FBodyInstance* GetBodyInstance(FName BoneName = NAME_None, bool bGetWelded = true, int32 Index = INDEX_NONE) const override;
+#if WITH_EDITOR
+	virtual void ReserveSurrogateInstanceBodyIndices(FFastGeoSurrogateBodyInstanceIndex& InOutNextAvailableInstanceBodyIndex) override;
+#endif
 	//~ End FFastGeoPrimitiveComponent interface
 
 	void ForEachInstanceMatrix(TFunctionRef<void(const FMatrix&)> Func) const;
@@ -50,6 +53,7 @@ protected:
 	virtual void ResetSceneProxyDescUnsupportedProperties() override;
 #endif
 	virtual void ApplyWorldTransform(const FTransform& InTransform) override;
+	virtual void GetBodyInstances(TArray<FBodyInstance*>& OutBodyInstances) override;
 	//~ End FFastGeoPrimitiveComponent interface
 
 	//~ Begin FFastGeoStaticMeshComponentBase interface
@@ -72,7 +76,7 @@ private:
 	};
 	FBoxSphereBounds CalculateBounds(EBoundsType BoundsType);
 	void CreateAllInstanceBodies();
-	static TArray<Chaos::FPhysicsObject*> GetAllPhysicsObjects(TArray<FBodyInstance*> InInstanceBodies);
+	static TArray<Chaos::FPhysicsObject*> GetAllPhysicsObjects(const TArray<FBodyInstance*>& InInstanceBodies);
 #if WITH_EDITOR
 	void InitializePerInstanceSMData(const TArray<FInstancedStaticMeshInstanceData>& InPerInstanceSMData);
 #endif
@@ -82,6 +86,7 @@ private:
 	TArray<FInstancedStaticMeshInstanceData> HighPrecisionPerInstanceSMData;
 	TArray<FTransform3f> LowPrecisionPerInstanceSMData;
 	bool bUseHighPrecisionPerInstanceSMData = false;
+	int32 LastInstanceBodyIndex = INDEX_NONE;
 	int32 InstancingRandomSeed = 0;
 	TArray<float> PerInstanceSMCustomData;
 	TArray<FInstancedStaticMeshRandomSeed> AdditionalRandomSeeds;
diff --git a/Engine/Plugins/Experimental/FastGeoStreaming/Source/FastGeoStreaming/Internal/FastGeoPrimitiveComponent.h b/Engine/Plugins/Experimental/FastGeoStreaming/Source/FastGeoStreaming/Internal/FastGeoPrimitiveComponent.h
--- a/Engine/Plugins/Experimental/FastGeoStreaming/Source/FastGeoStreaming/Internal/FastGeoPrimitiveComponent.h
+++ b/Engine/Plugins/Experimental/FastGeoStreaming/Source/FastGeoStreaming/Internal/FastGeoPrimitiveComponent.h
@@ -17,8 +17,10 @@ class FPrimitiveSceneProxy;
 class UPrimitiveComponent;
 class URuntimeVirtualTexture;
 class UFastGeoContainer;
+class UFastGeoSurrogateComponent;
 class FFastGeoPrimitiveComponent;
 class FFastGeoDestroyRenderStateContext;
+struct FFastGeoSurrogateBodyInstanceIndex;
 struct FPrimitiveSceneDesc;
 struct FPrimitiveSceneProxyDesc;
 struct FPSOPrecacheParams;
@@ -135,7 +137,10 @@ public:
 	virtual bool DoCustomNavigableGeometryExport(FNavigableGeometryExport& GeomExport) const { return true; }
 	//~ End Navigation
 
-	
+#if WITH_EDITOR
+	virtual void ReserveSurrogateInstanceBodyIndices(FFastGeoSurrogateBodyInstanceIndex& InOutNextAvailableInstanceBodyIndex);
+#endif
+
 protected:
 	virtual FPrimitiveSceneProxyDesc& GetSceneProxyDesc() = 0;
 	virtual const FPrimitiveSceneProxyDesc& GetSceneProxyDesc() const = 0;
@@ -173,6 +178,7 @@ protected:
 	uint8 bRasterizeAsFilledConvexVolume : 1 = false;
 	uint8 bCanEverAffectNavigation : 1 = false;
 	uint8 bMultiBodyOverlap : 1 = false;
+	int32 SurrogateComponentDescriptorIndex = INDEX_NONE;
 	FCustomPrimitiveData CustomPrimitiveData;
 	TEnumAsByte<EHasCustomNavigableGeometry::Type> bHasCustomNavigableGeometry = EHasCustomNavigableGeometry::No;
 	FBodyInstance BodyInstance;
@@ -192,10 +198,12 @@ protected:
 	/** All PSO related data for the component */
 	FPSOPrecacheComponentData PSOPrecacheComponentData;
 #endif
+	UFastGeoSurrogateComponent* GetSurrogateComponent() const;
 
 private:
 	bool UsePSOPrecacheFallbackMaterial() const;
 
+	friend class UFastGeoWorldPartitionRuntimeCellTransformer;
 	friend class FFastGeoComponentCluster;
 	friend class FPrimitiveComponentHelper;
 	friend class FFastGeoPhysicsBodyInstanceOwner;
diff --git a/Engine/Plugins/Experimental/FastGeoStreaming/Source/FastGeoStreaming/Internal/FastGeoStaticMeshComponent.h b/Engine/Plugins/Experimental/FastGeoStreaming/Source/FastGeoStreaming/Internal/FastGeoStaticMeshComponent.h
--- a/Engine/Plugins/Experimental/FastGeoStreaming/Source/FastGeoStreaming/Internal/FastGeoStaticMeshComponent.h
+++ b/Engine/Plugins/Experimental/FastGeoStreaming/Source/FastGeoStreaming/Internal/FastGeoStaticMeshComponent.h
@@ -144,6 +144,7 @@ protected:
 	//~ Begin FFastGeoPrimitiveComponent interface
 	virtual FPrimitiveSceneProxyDesc& GetSceneProxyDesc() override { return SceneProxyDesc; }
 	virtual const FPrimitiveSceneProxyDesc& GetSceneProxyDesc() const override { return SceneProxyDesc; }
+	virtual void GetBodyInstances(TArray<FBodyInstance*>& OutBodyInstances) override;
 	//~ End FFastGeoPrimitiveComponent interface
 
 	//~ Begin FFastGeoStaticMeshComponentBase interface
diff --git a/Engine/Plugins/Experimental/FastGeoStreaming/Source/FastGeoStreaming/Internal/FastGeoSurrogateActor.h b/Engine/Plugins/Experimental/FastGeoStreaming/Source/FastGeoStreaming/Internal/FastGeoSurrogateActor.h
new file mode 100644
--- /dev/null
+++ b/Engine/Plugins/Experimental/FastGeoStreaming/Source/FastGeoStreaming/Internal/FastGeoSurrogateActor.h
@@ -0,0 +1,53 @@
+// Copyright Epic Games, Inc. All Rights Reserved.
+
+#pragma once
+
+#include "GameFramework/Actor.h"
+#include "FastGeoSurrogateActor.generated.h"
+
+class UFastGeoSurrogateComponent;
+struct FFastGeoSurrogateComponentDescriptor;
+
+UCLASS(NotBlueprintable)
+class FASTGEOSTREAMING_API AFastGeoSurrogateActor : public AActor
+{
+	GENERATED_BODY()
+
+public:
+
+	AFastGeoSurrogateActor();
+
+	const FFastGeoSurrogateComponentDescriptor* GetSurrogateComponentDescriptor(int32 SurrogateComponentDescriptorIndex) const;
+	UFastGeoSurrogateComponent* GetSurrogateComponent(int32 SurrogateComponentDescriptorIndex) const;
+	bool HasRegisteredWithFastGeo() const { return bHasRegisteredWithFastGeo; }
+	bool IsActive() const { return bIsActive; }
+
+protected:
+
+	//~ Begin AActor interface
+	virtual void PostInitializeComponents() override;
+	virtual void PostUnregisterAllComponents() override;
+	virtual void EndPlay(const EEndPlayReason::Type EndPlayReason) override;
+	//~ End AActor interface
+
+private:
+
+	bool UnregisterFromFastGeo();
+
+	/** Maps surrogate component descriptor index to surrogate component descriptor */
+	UPROPERTY()
+	TArray<FFastGeoSurrogateComponentDescriptor> SurrogateComponentDescriptors;
+
+	/** Maps surrogate component descriptor index to surrogate component */
+	UPROPERTY(Transient)
+	TArray<TObjectPtr<UFastGeoSurrogateComponent>> SurrogateComponents;
+
+	/** Whether actor has registered with FastGeo container */
+	UPROPERTY(Transient)
+	bool bHasRegisteredWithFastGeo = false;
+
+	UPROPERTY(Transient)
+	bool bIsActive = false;
+
+	friend class UFastGeoWorldPartitionRuntimeCellTransformer;
+};
\ No newline at end of file
diff --git a/Engine/Plugins/Experimental/FastGeoStreaming/Source/FastGeoStreaming/Internal/FastGeoSurrogateBodyInstanceIndex.h b/Engine/Plugins/Experimental/FastGeoStreaming/Source/FastGeoStreaming/Internal/FastGeoSurrogateBodyInstanceIndex.h
new file mode 100644
--- /dev/null
+++ b/Engine/Plugins/Experimental/FastGeoStreaming/Source/FastGeoStreaming/Internal/FastGeoSurrogateBodyInstanceIndex.h
@@ -0,0 +1,98 @@
+// Copyright Epic Games, Inc. All Rights Reserved.
+
+#pragma once
+
+#include "CoreMinimal.h"
+
+struct FFastGeoSurrogateBodyInstanceIndex
+{
+	static FFastGeoSurrogateBodyInstanceIndex FromEncoded(int32 InEncodedIndex)
+	{
+		check(InEncodedIndex == INDEX_NONE || IsEncoded(InEncodedIndex));
+		return FFastGeoSurrogateBodyInstanceIndex(InEncodedIndex);
+	}
+
+	static FFastGeoSurrogateBodyInstanceIndex Encode(int32 InRawIndex)
+	{
+		check(InRawIndex == INDEX_NONE || InRawIndex >= 0);
+		check(InRawIndex == INDEX_NONE || ((InRawIndex & BodyInstanceIndexEncodeBit) == 0));
+		return (InRawIndex == INDEX_NONE)
+			? FFastGeoSurrogateBodyInstanceIndex(INDEX_NONE)
+			: FFastGeoSurrogateBodyInstanceIndex(InRawIndex | BodyInstanceIndexEncodeBit);
+	}
+
+	static int32 Decode(int32 Index)
+	{
+		return (Index == INDEX_NONE) ? INDEX_NONE : (Index & ~BodyInstanceIndexEncodeBit);
+	}
+
+	static bool IsEncoded(int32 Index)
+	{
+		return (Index != INDEX_NONE) && (Index >= 0) && ((Index & BodyInstanceIndexEncodeBit) != 0);
+	}
+
+	int32 Decode() const
+	{
+		return Decode(EncodedIndex);
+	}
+
+	bool IsEncoded() const
+	{
+		return IsEncoded(EncodedIndex);
+	}
+
+	int32 GetEncoded() const
+	{
+		return EncodedIndex;
+	}
+
+	explicit operator int32() const
+	{
+		return EncodedIndex;
+	}
+
+	FFastGeoSurrogateBodyInstanceIndex& operator++()
+	{
+		return (*this += 1);
+	}
+
+	FFastGeoSurrogateBodyInstanceIndex operator++(int)
+	{
+		FFastGeoSurrogateBodyInstanceIndex Temp = *this;
+		++(*this);
+		return Temp;
+	}
+
+	FFastGeoSurrogateBodyInstanceIndex& operator+=(int32 Count)
+	{
+		check(Count >= 0);
+		check(EncodedIndex != INDEX_NONE);
+
+		const int64 Raw = (int64)Decode();
+		const int64 NewRaw = Raw + (int64)Count;
+
+		check(NewRaw >= 0);
+		check(NewRaw < (int64)BodyInstanceIndexEncodeBit);
+
+		EncodedIndex = (int32)NewRaw | BodyInstanceIndexEncodeBit;
+		return *this;
+	}
+
+	friend FFastGeoSurrogateBodyInstanceIndex operator+(FFastGeoSurrogateBodyInstanceIndex Lhs, int32 Count)
+	{
+		Lhs += Count;
+		return Lhs;
+	}
+
+private:
+	explicit FFastGeoSurrogateBodyInstanceIndex(int32 InEncoded)
+		: EncodedIndex(InEncoded)
+	{
+		check(InEncoded == INDEX_NONE || IsEncoded(InEncoded));
+	}
+
+	int32 EncodedIndex = INDEX_NONE;
+
+	/** Bit 30 used to flag that a BodyInstance Index has been encoded. */
+	static constexpr int32 BodyInstanceIndexEncodeBit = 1 << 30;
+};
\ No newline at end of file
diff --git a/Engine/Plugins/Experimental/FastGeoStreaming/Source/FastGeoStreaming/Internal/FastGeoSurrogateComponent.h b/Engine/Plugins/Experimental/FastGeoStreaming/Source/FastGeoStreaming/Internal/FastGeoSurrogateComponent.h
new file mode 100644
--- /dev/null
+++ b/Engine/Plugins/Experimental/FastGeoStreaming/Source/FastGeoStreaming/Internal/FastGeoSurrogateComponent.h
@@ -0,0 +1,62 @@
+// Copyright Epic Games, Inc. All Rights Reserved.
+
+#pragma once
+
+#include "PhysicsEngine/BodyInstance.h"
+#include "FastGeoPrimitiveComponent.h"
+
+#include "FastGeoSurrogateComponent.generated.h"
+
+struct FFastGeoSurrogateComponentDescriptor;
+
+UCLASS(NotBlueprintable)
+class FASTGEOSTREAMING_API UFastGeoSurrogateComponent : public UPrimitiveComponent
+{
+	GENERATED_BODY()
+
+public:
+
+	UFastGeoSurrogateComponent();
+
+	//~ Begin UActorComponent interface
+	virtual void OnRegister() override;
+	virtual bool ShouldCreatePhysicsState() const override;
+	virtual void OnCreatePhysicsState() override;
+	//~ End UActorComponent interface
+
+	//~ Begin UPrimitiveComponent interface
+	virtual FPrimitiveSceneProxy* CreateSceneProxy() override { return nullptr; }
+	virtual FBoxSphereBounds CalcBounds(const FTransform& LocalToWorld) const override;
+	virtual FBodyInstance* GetBodyInstance(FName BoneName = NAME_None, bool bGetWelded = true, int32 Index = INDEX_NONE) const override;
+	//~ End UPrimitiveComponent interface
+
+	//~ Begin USceneComponent interface
+	virtual bool IsSimulatingPhysics(FName BoneName = NAME_None) const override { return false; }
+	virtual ECollisionEnabled::Type GetCollisionEnabled() const override;
+	virtual ECollisionResponse GetCollisionResponseToChannel(ECollisionChannel Channel) const override;
+	virtual const FCollisionResponseContainer& GetCollisionResponseToChannels() const override;
+	//~ End USceneComponent interface
+
+	//~ Begin IPhysicsBodyInstanceOwner interface
+	virtual bool IsMultiBodyOverlap() const override { return true; }
+	//~ End IPhysicsBodyInstanceOwner interface
+
+	void RegisterBodyInstances(const TArray<FBodyInstance*>& InBodyInstances);
+	void UnregisterBodyInstances();
+
+protected:
+
+	void SetSurrogateComponentDescriptor(const FFastGeoSurrogateComponentDescriptor& InComponentDesc);
+	int32 GetDescriptorIndex() const { return DescriptorIndex; }
+
+	UPROPERTY()
+	int32 DescriptorIndex = INDEX_NONE;
+
+	TArray<FBodyInstance*> IndexToBodyInstance;
+
+	int32 RegisteredBodyInstanceCount = 0;
+
+	friend class FFastGeoPhysicsBodyInstanceOwner;
+	friend class AFastGeoSurrogateActor;
+	friend class UFastGeoWorldPartitionRuntimeCellTransformer;
+};
\ No newline at end of file
diff --git a/Engine/Plugins/Experimental/FastGeoStreaming/Source/FastGeoStreaming/Internal/FastGeoSurrogateComponentDescriptor.h b/Engine/Plugins/Experimental/FastGeoStreaming/Source/FastGeoStreaming/Internal/FastGeoSurrogateComponentDescriptor.h
new file mode 100644
--- /dev/null
+++ b/Engine/Plugins/Experimental/FastGeoStreaming/Source/FastGeoStreaming/Internal/FastGeoSurrogateComponentDescriptor.h
@@ -0,0 +1,56 @@
+// Copyright Epic Games, Inc. All Rights Reserved.
+
+#pragma once
+
+#include "Engine/EngineTypes.h"
+#include "Components/PrimitiveComponent.h"
+#include "FastGeoSurrogateComponentDescriptor.generated.h"
+
+USTRUCT()
+struct FFastGeoSurrogateComponentDescriptor
+{
+	GENERATED_BODY()
+
+	UPROPERTY()
+	TEnumAsByte<ECollisionEnabled::Type> CollisionEnabled;
+
+	UPROPERTY()
+	TEnumAsByte<ECollisionChannel> ObjectType;
+
+	UPROPERTY()
+	FCollisionResponseContainer ResponseContainer;
+
+	UPROPERTY()
+	FWalkableSlopeOverride WalkableSlopeOverride;
+
+	UPROPERTY()
+	TEnumAsByte<ECanBeCharacterBase> CanCharacterStepUpOn;
+
+	FFastGeoSurrogateComponentDescriptor()
+		: CollisionEnabled(ECollisionEnabled::NoCollision)
+		, ObjectType(ECollisionChannel::ECC_WorldStatic)
+		, CanCharacterStepUpOn(ECanBeCharacterBase::ECB_No)
+	{
+	}
+
+#if WITH_EDITOR
+	FFastGeoSurrogateComponentDescriptor(const UPrimitiveComponent& InPrimitiveComponent)
+		: CollisionEnabled(InPrimitiveComponent.BodyInstance.GetCollisionEnabled())
+		, ObjectType(InPrimitiveComponent.BodyInstance.GetObjectType())
+		, ResponseContainer(InPrimitiveComponent.BodyInstance.GetResponseToChannels())
+		, WalkableSlopeOverride(InPrimitiveComponent.BodyInstance.GetWalkableSlopeOverride())
+		, CanCharacterStepUpOn(InPrimitiveComponent.CanCharacterStepUpOn)
+	{
+	}
+#endif
+
+	FORCEINLINE bool operator==(const FFastGeoSurrogateComponentDescriptor& Other) const
+	{
+		return CollisionEnabled == Other.CollisionEnabled &&
+			ObjectType == Other.ObjectType &&
+			ResponseContainer == Other.ResponseContainer &&
+			WalkableSlopeOverride.WalkableSlopeBehavior == Other.WalkableSlopeOverride.WalkableSlopeBehavior &&
+			WalkableSlopeOverride.WalkableSlopeAngle == Other.WalkableSlopeOverride.WalkableSlopeAngle &&
+			CanCharacterStepUpOn == Other.CanCharacterStepUpOn;
+	}
+};
\ No newline at end of file
diff --git a/Engine/Plugins/Experimental/FastGeoStreaming/Source/FastGeoStreaming/Internal/FastGeoWorldSubsystem.h b/Engine/Plugins/Experimental/FastGeoStreaming/Source/FastGeoStreaming/Internal/FastGeoWorldSubsystem.h
--- a/Engine/Plugins/Experimental/FastGeoStreaming/Source/FastGeoStreaming/Internal/FastGeoWorldSubsystem.h
+++ b/Engine/Plugins/Experimental/FastGeoStreaming/Source/FastGeoStreaming/Internal/FastGeoWorldSubsystem.h
@@ -78,6 +78,7 @@ public:
 	void PushAsyncDestroyPhysicsStateJobs(UFastGeoContainer* FastGeo);
 
 	static bool IsEnableDebugView();
+	static bool ShouldAllowSurrogateComponents();
 
 #if !UE_BUILD_SHIPPING
 	static bool IsFastGeoVisible();
@@ -130,7 +131,6 @@ private:
 
 	TArray<FFastGeoRegisteredComponent> ComponentsPendingRecreate;
 	static bool bEnableDebugView;
-	static bool bIsFastGeoVisible;
 
 	FRWLock Lock;
 	int32 TimeEpoch = 0;
diff --git a/Engine/Plugins/Experimental/FastGeoStreaming/Source/FastGeoStreaming/Private/FastGeoComponent.cpp b/Engine/Plugins/Experimental/FastGeoStreaming/Source/FastGeoStreaming/Private/FastGeoComponent.cpp
--- a/Engine/Plugins/Experimental/FastGeoStreaming/Source/FastGeoStreaming/Private/FastGeoComponent.cpp
+++ b/Engine/Plugins/Experimental/FastGeoStreaming/Source/FastGeoStreaming/Private/FastGeoComponent.cpp
@@ -27,7 +27,9 @@ FFastGeoComponent::FFastGeoComponent(const FFastGeoComponent& Other)
 , DetailMode(Other.DetailMode)
 , Owner(Other.Owner)
 , Lock(MakeUnique<FRWLock>())
+#ifdef FASTGEO_DEBUG_PHYSICS
 , PhysicsStateCreation(EPhysicsStateCreation::NotCreated)
+#endif
 , ProxyState(EProxyCreationState::None)
 , bRenderStateDirty(false)
 {
@@ -108,6 +110,7 @@ void FFastGeoComponent::ApplyWorldTransform(const FTransform& InTransform)
 	WorldTransform = LocalTransform * InTransform;
 }
 
+#ifdef FASTGEO_DEBUG_PHYSICS
 void FFastGeoComponent::OnAsyncCreatePhysicsState()
 {
 	check(IsCollisionEnabled());
@@ -137,6 +140,7 @@ void FFastGeoComponent::OnAsyncDestroyPhysicsStateEnd_GameThread()
 	check(PhysicsStateCreation == EPhysicsStateCreation::Destroying);
 	PhysicsStateCreation = EPhysicsStateCreation::NotCreated;
 }
+#endif
 
 bool FFastGeoComponent::IsRenderStateCreated() const
 {
diff --git a/Engine/Plugins/Experimental/FastGeoStreaming/Source/FastGeoStreaming/Private/FastGeoContainer.cpp b/Engine/Plugins/Experimental/FastGeoStreaming/Source/FastGeoStreaming/Private/FastGeoContainer.cpp
--- a/Engine/Plugins/Experimental/FastGeoStreaming/Source/FastGeoStreaming/Private/FastGeoContainer.cpp
+++ b/Engine/Plugins/Experimental/FastGeoStreaming/Source/FastGeoStreaming/Private/FastGeoContainer.cpp
@@ -5,6 +5,8 @@
 #include "FastGeoHLOD.h"
 #include "FastGeoComponent.h"
 #include "FastGeoStreamingModule.h"
+#include "FastGeoSurrogateActor.h"
+#include "FastGeoSurrogateComponent.h"
 #include "FastGeoWorldSubsystem.h"
 #include "FastGeoDestroyRenderStateContext.h"
 #include "FastGeoLog.h"
@@ -185,6 +187,14 @@ void UFastGeoContainer::Register()
 
 	if (!bIsRegistered)
 	{
+		check(!IsUsingSurrogateComponents() || !SurrogateActor->HasRegisteredWithFastGeo());
+
+		if (const bool bHasAnyPendingTasks = HasAnyPendingTasks())
+		{
+			checkf(!bHasAnyPendingTasks, TEXT("UFastGeoContainer::Register() was called while container still has pending tasks. This will lead to invalid behavior."));
+			return;
+		}
+
 		++RegistrationEpoch;
 		PendingCreate.Reset();
 
@@ -195,7 +205,7 @@ void UFastGeoContainer::Register()
 		UWorldPartition* WorldPartition = FWorldPartitionHelpers::GetWorldPartition(this);
 		const bool bApplyWorldTransform = WorldPartition && WorldPartition->HasInstanceTransform();
 		const FTransform& Transform = WorldPartition ? WorldPartition->GetInstanceTransform() : FTransform::Identity;
-		ForEachComponentCluster([bApplyWorldTransform , &Transform, this](FFastGeoComponentCluster& ComponentCluster)
+		ForEachComponentCluster([bApplyWorldTransform, &Transform, this](FFastGeoComponentCluster& ComponentCluster)
 		{
 			ComponentCluster.OnRegister();
 
@@ -317,6 +327,14 @@ void UFastGeoContainer::Unregister()
 
 	if (bIsRegistered)
 	{
+		check(!IsUsingSurrogateComponents() || SurrogateActor->HasRegisteredWithFastGeo());
+
+		if (const bool bHasAnyPendingTasks = HasAnyPendingTasks())
+		{
+			checkf(!bHasAnyPendingTasks, TEXT("UFastGeoContainer::Unregister() was called while container still has pending tasks. This will lead to invalid behavior."));
+			return;
+		}
+
 		PendingDestroy.Reset();
 
 		ForEachComponentCluster([this](FFastGeoComponentCluster& ComponentCluster)
@@ -946,11 +964,15 @@ void UFastGeoContainer::OnAsyncCreatePhysicsStateEnd_GameThread()
 	check(PhysicsStateCreation == EPhysicsStateCreation::Creating);
 	PhysicsStateCreation = EPhysicsStateCreation::Created;
 	
-	// This loop could be removed if component's PhysicsStateCreation was removed
+#ifdef FASTGEO_DEBUG_PHYSICS
+	// This loop is only necessary for debugging
 	for (FFastGeoComponent* Component : CollisionComponents)
 	{
 		Component->OnAsyncCreatePhysicsStateEnd_GameThread();
 	}
+#endif
+
+	RegisterBodyInstances();
 }
 
 void UFastGeoContainer::OnAsyncDestroyPhysicsStateBegin_GameThread()
@@ -965,6 +987,8 @@ void UFastGeoContainer::OnAsyncDestroyPhysicsStateBegin_GameThread()
 	{
 		Component->OnAsyncDestroyPhysicsStateBegin_GameThread();
 	}
+
+	UnregisterBodyInstances();
 }
 
 bool UFastGeoContainer::OnAsyncDestroyPhysicsState(const UE::FTimeout& Timeout)
@@ -992,7 +1016,6 @@ void UFastGeoContainer::OnAsyncDestroyPhysicsStateEnd_GameThread()
 	check(PhysicsStateCreation == EPhysicsStateCreation::Destroying);
 	PhysicsStateCreation = EPhysicsStateCreation::NotCreated;
 
-	// This loop could be removed if component's PhysicsStateCreation was removed
 	for (FFastGeoComponent* Component : CollisionComponents)
 	{
 		Component->OnAsyncDestroyPhysicsStateEnd_GameThread();
@@ -1020,3 +1043,48 @@ IPhysicsBodyInstanceOwner* UFastGeoContainer::ResolvePhysicsBodyInstanceOwner(Ch
 	}
 	return nullptr;
 }
+
+bool UFastGeoContainer::IsUsingSurrogateComponents() const
+{
+	return IsValid(SurrogateActor) && UFastGeoWorldSubsystem::ShouldAllowSurrogateComponents();
+}
+
+AFastGeoSurrogateActor* UFastGeoContainer::GetSurrogateActor() const
+{
+	return SurrogateActor;
+}
+
+void UFastGeoContainer::RegisterBodyInstances()
+{
+	if (IsUsingSurrogateComponents())
+	{
+		TMap<UFastGeoSurrogateComponent*, TArray<FBodyInstance*>> SurrogateBatches;
+		for (FFastGeoComponent* Component : CollisionComponents)
+		{
+			if (FFastGeoPrimitiveComponent* PrimComp = Component->CastTo<FFastGeoPrimitiveComponent>())
+			{
+				if (UFastGeoSurrogateComponent* SurrogateComp = SurrogateActor->GetSurrogateComponent(PrimComp->SurrogateComponentDescriptorIndex))
+				{
+					TArray<FBodyInstance*>& Batch = SurrogateBatches.FindOrAdd(SurrogateComp);
+					PrimComp->GetBodyInstances(Batch);
+				}
+			}
+		}
+		for (TPair<UFastGeoSurrogateComponent*, TArray<FBodyInstance*>>& Pair : SurrogateBatches)
+		{
+			UFastGeoSurrogateComponent* SurrogateComponent = Pair.Key;
+			SurrogateComponent->RegisterBodyInstances(Pair.Value);
+		}
+	}
+}
+
+void UFastGeoContainer::UnregisterBodyInstances()
+{
+	if (IsUsingSurrogateComponents())
+	{
+		SurrogateActor->ForEachComponent<UFastGeoSurrogateComponent>(false, [](UFastGeoSurrogateComponent* SurrogateComponent)
+		{
+			SurrogateComponent->UnregisterBodyInstances();
+		});
+	}
+}
\ No newline at end of file
diff --git a/Engine/Plugins/Experimental/FastGeoStreaming/Source/FastGeoStreaming/Private/FastGeoInstancedStaticMeshComponent.cpp b/Engine/Plugins/Experimental/FastGeoStreaming/Source/FastGeoStreaming/Private/FastGeoInstancedStaticMeshComponent.cpp
--- a/Engine/Plugins/Experimental/FastGeoStreaming/Source/FastGeoStreaming/Private/FastGeoInstancedStaticMeshComponent.cpp
+++ b/Engine/Plugins/Experimental/FastGeoStreaming/Source/FastGeoStreaming/Private/FastGeoInstancedStaticMeshComponent.cpp
@@ -3,6 +3,8 @@
 #include "FastGeoInstancedStaticMeshComponent.h"
 #include "FastGeoWorldPartitionRuntimeCellTransformer.h"
 #include "FastGeoComponentCluster.h"
+#include "FastGeoSurrogateComponent.h"
+#include "FastGeoSurrogateBodyInstanceIndex.h"
 #include "FastGeoContainer.h"
 #include "FastGeoHLOD.h"
 #include "FastGeoLog.h"
@@ -39,6 +41,7 @@ void FFastGeoInstancedStaticMeshComponent::Serialize(FArchive& Ar)
 	{
 		LowPrecisionPerInstanceSMData.BulkSerialize(Ar);
 	}
+	Ar << LastInstanceBodyIndex;
 	Ar << InstancingRandomSeed;
 	PerInstanceSMCustomData.BulkSerialize(Ar);
 	Ar << AdditionalRandomSeeds;
@@ -333,7 +336,7 @@ TArray<Chaos::FPhysicsObject*> FFastGeoInstancedStaticMeshComponent::GetAllPhysi
 	return GetAllPhysicsObjects(InstanceBodies);
 }
 
-TArray<Chaos::FPhysicsObject*> FFastGeoInstancedStaticMeshComponent::GetAllPhysicsObjects(TArray<FBodyInstance*> InInstanceBodies)
+TArray<Chaos::FPhysicsObject*> FFastGeoInstancedStaticMeshComponent::GetAllPhysicsObjects(const TArray<FBodyInstance*>& InInstanceBodies)
 {
 	TArray<Chaos::FPhysicsObject*> Objects;
 	Objects.Reserve(InInstanceBodies.Num());
@@ -349,9 +352,24 @@ TArray<Chaos::FPhysicsObject*> FFastGeoInstancedStaticMeshComponent::GetAllPhysi
 
 FBodyInstance* FFastGeoInstancedStaticMeshComponent::GetBodyInstance(FName BoneName, bool bGetWelded, int32 Index) const
 {
-	if (Index != INDEX_NONE && InstanceBodies.IsValidIndex(Index))
+	if (Index != INDEX_NONE)
 	{
-		return const_cast<FBodyInstance*>(InstanceBodies[Index]);
+		// Handle the case where GetBodyInstance is called directly on the FastGeo component with an already transformed Index
+		if (UFastGeoSurrogateComponent* SurrogateComponent = GetSurrogateComponent())
+		{
+			if (FFastGeoSurrogateBodyInstanceIndex::IsEncoded(Index))
+			{
+				return SurrogateComponent->GetBodyInstance(BoneName, bGetWelded, Index);
+			}
+		}
+		else
+		{
+			check(BodyInstance.InstanceBodyIndex == INDEX_NONE);
+			if (InstanceBodies.IsValidIndex(Index))
+			{
+				return const_cast<FBodyInstance*>(InstanceBodies[Index]);
+			}
+		}
 	}
 	return const_cast<FBodyInstance*>(&BodyInstance);
 }
@@ -379,10 +397,15 @@ void FFastGeoInstancedStaticMeshComponent::CreateAllInstanceBodies()
 		TArray<FBodyInstance*> InstanceBodiesSanitized;
 		InstanceBodiesSanitized.Reserve(NumBodies);
 
+		UFastGeoSurrogateComponent* SurrogateComponent = GetSurrogateComponent();
+		const bool bUseSurrogateComponent = !!SurrogateComponent;
+		check(!bUseSurrogateComponent || FFastGeoSurrogateBodyInstanceIndex::IsEncoded(BodyInstance.InstanceBodyIndex));
+		const int32 BaseOffset = bUseSurrogateComponent ? BodyInstance.InstanceBodyIndex : 0;
+
 		TArray<FTransform> Transforms;
 		Transforms.Reserve(NumBodies);
 		int32 InstanceBodyIndex = 0;
-		ForEachInstanceMatrix([this, &InstanceBodyIndex, &Transforms, &InstanceBodiesSanitized](const FMatrix& InstanceTransform)
+		ForEachInstanceMatrix([this, BaseOffset, &InstanceBodyIndex, &Transforms, &InstanceBodiesSanitized](const FMatrix& InstanceTransform)
 		{
 			const FTransform InstanceTM = FTransform(InstanceTransform) * WorldTransform;
 			if (InstanceTM.GetScale3D().IsNearlyZero())
@@ -395,7 +418,7 @@ void FFastGeoInstancedStaticMeshComponent::CreateAllInstanceBodies()
 				InstanceBodiesSanitized.Add(Instance);
 				InstanceBodies[InstanceBodyIndex] = Instance;
 				Instance->CopyBodyInstancePropertiesFrom(&BodyInstance);
-				Instance->InstanceBodyIndex = InstanceBodyIndex;
+				Instance->InstanceBodyIndex = BaseOffset + InstanceBodyIndex;
 				Instance->bAutoWeld = false;
 				Instance->bSimulatePhysics = false;
 				Transforms.Add(InstanceTM);
@@ -403,13 +426,15 @@ void FFastGeoInstancedStaticMeshComponent::CreateAllInstanceBodies()
 			++InstanceBodyIndex;
 		});
 
+		check(!bUseSurrogateComponent || InstanceBodiesSanitized.IsEmpty() || InstanceBodiesSanitized.Last()->InstanceBodyIndex <= LastInstanceBodyIndex);
+
 		if (InstanceBodiesSanitized.Num() > 0)
 		{
 			// Initialize BodyInstanceOwner
 			BodyInstanceOwner.Initialize(this);
 
 			// Initialize body instances
-			FBodyInstance::InitStaticBodies(MoveTemp(InstanceBodiesSanitized), MoveTemp(Transforms), BodySetup, nullptr, GetWorld()->GetPhysicsScene(), &BodyInstanceOwner);
+			FBodyInstance::InitStaticBodies(MoveTemp(InstanceBodiesSanitized), MoveTemp(Transforms), BodySetup, SurrogateComponent, GetWorld()->GetPhysicsScene(), &BodyInstanceOwner);
 
 			// Assign BodyInstanceOwner
 			TArray<Chaos::FPhysicsObject*> PhysicsObjects = GetAllPhysicsObjects(InstanceBodies);
@@ -483,4 +508,29 @@ void FFastGeoInstancedStaticMeshComponent::CollectPSOPrecacheData(const FPSOPrec
 {
 	return FInstancedStaticMeshComponentHelper::CollectPSOPrecacheData(*this, BasePrecachePSOParams, OutParams);
 }
+#endif
+
+void FFastGeoInstancedStaticMeshComponent::GetBodyInstances(TArray<FBodyInstance*>& OutBodyInstances)
+{
+	OutBodyInstances.Append(InstanceBodies);
+}
+
+#if WITH_EDITOR
+void FFastGeoInstancedStaticMeshComponent::ReserveSurrogateInstanceBodyIndices(FFastGeoSurrogateBodyInstanceIndex& InOutNextAvailableInstanceBodyIndex)
+{
+	const int32 InstanceCount = GetInstanceCount();
+	if (InstanceCount > 0)
+	{
+		const int32 RangeOffset = InstanceCount - 1;
+		Super::ReserveSurrogateInstanceBodyIndices(InOutNextAvailableInstanceBodyIndex);
+		check(FFastGeoSurrogateBodyInstanceIndex::IsEncoded(BodyInstance.InstanceBodyIndex));
+		LastInstanceBodyIndex = (FFastGeoSurrogateBodyInstanceIndex::FromEncoded(BodyInstance.InstanceBodyIndex) + RangeOffset).GetEncoded();
+		InOutNextAvailableInstanceBodyIndex += RangeOffset;
+	}
+	else
+	{
+		check(BodyInstance.InstanceBodyIndex == INDEX_NONE);
+		check(LastInstanceBodyIndex == INDEX_NONE);
+	}
+}
 #endif
\ No newline at end of file
diff --git a/Engine/Plugins/Experimental/FastGeoStreaming/Source/FastGeoStreaming/Private/FastGeoPrimitiveComponent.cpp b/Engine/Plugins/Experimental/FastGeoStreaming/Source/FastGeoStreaming/Private/FastGeoPrimitiveComponent.cpp
--- a/Engine/Plugins/Experimental/FastGeoStreaming/Source/FastGeoStreaming/Private/FastGeoPrimitiveComponent.cpp
+++ b/Engine/Plugins/Experimental/FastGeoStreaming/Source/FastGeoStreaming/Private/FastGeoPrimitiveComponent.cpp
@@ -6,6 +6,9 @@
 #include "FastGeoInstancedStaticMeshComponent.h"
 #include "FastGeoWorldSubsystem.h"
 #include "FastGeoComponentCluster.h"
+#include "FastGeoSurrogateBodyInstanceIndex.h"
+#include "FastGeoSurrogateComponent.h"
+#include "FastGeoSurrogateActor.h"
 #include "FastGeoDestroyRenderStateContext.h"
 #include "FastGeoWeakElement.h"
 #include "FastGeoLog.h"
@@ -48,6 +51,7 @@ FFastGeoPrimitiveComponent::FFastGeoPrimitiveComponent(const FFastGeoPrimitiveCo
 	, bRasterizeAsFilledConvexVolume(Other.bRasterizeAsFilledConvexVolume)
 	, bCanEverAffectNavigation(Other.bCanEverAffectNavigation)
 	, bMultiBodyOverlap(Other.bMultiBodyOverlap)
+	, SurrogateComponentDescriptorIndex(Other.SurrogateComponentDescriptorIndex)
 	, CustomPrimitiveData(Other.CustomPrimitiveData)
 	, bHasCustomNavigableGeometry(Other.bHasCustomNavigableGeometry)
 	, BodyInstance(Other.BodyInstance)
@@ -71,10 +75,17 @@ void FFastGeoPrimitiveComponent::Serialize(FArchive& Ar)
 	FArchive_Serialize_BitfieldBool(Ar, bRasterizeAsFilledConvexVolume);
 	FArchive_Serialize_BitfieldBool(Ar, bCanEverAffectNavigation);
 	FArchive_Serialize_BitfieldBool(Ar, bMultiBodyOverlap);
+	Ar << SurrogateComponentDescriptorIndex;
 	Ar << CustomPrimitiveData.Data;
 	Ar << bHasCustomNavigableGeometry;
 	Ar << RuntimeVirtualTextures;
 	FBodyInstance::StaticStruct()->SerializeItem(Ar, &BodyInstance, nullptr);
+	Ar << BodyInstance.InstanceBodyIndex;
+
+	if (Ar.IsLoading() && !UFastGeoWorldSubsystem::ShouldAllowSurrogateComponents())
+	{
+		BodyInstance.InstanceBodyIndex = INDEX_NONE;
+	}
 
 	// Serialize persistent data from FPrimitiveSceneProxyDesc
 	FPrimitiveSceneProxyDesc& SceneProxyDesc = GetSceneProxyDesc();
@@ -164,6 +175,7 @@ void FFastGeoPrimitiveComponent::InitializeFromComponent(UActorComponent* Compon
 	bRasterizeAsFilledConvexVolume = PrimitiveComponent->bRasterizeAsFilledConvexVolume;
 	bCanEverAffectNavigation = PrimitiveComponent->CanEverAffectNavigation();
 	bMultiBodyOverlap = PrimitiveComponent->bMultiBodyOverlap;
+	SurrogateComponentDescriptorIndex = INDEX_NONE;
 	CustomPrimitiveData = PrimitiveComponent->GetCustomPrimitiveData();
 	bHasCustomNavigableGeometry = PrimitiveComponent->bHasCustomNavigableGeometry;
 	BodyInstance.CopyBodyInstancePropertiesFrom(&PrimitiveComponent->BodyInstance);
@@ -413,8 +425,12 @@ void FFastGeoPrimitiveComponent::OnAsyncCreatePhysicsState()
 			// Initialize BodyInstanceOwner
 			BodyInstanceOwner.Initialize(this);
 
+			UFastGeoSurrogateComponent* SurrogateComponent = GetSurrogateComponent();
+			// Make sure that the BodyInstance.InstanceBodyIndex was properly set by the FastGeo runtime cell transformer
+			check(!SurrogateComponent || FFastGeoSurrogateBodyInstanceIndex::IsEncoded(BodyInstance.InstanceBodyIndex));
+
 			// Initialize the body instance
-			BodyInstance.InitBody(BodySetup, BodyTransform, nullptr, GetWorld()->GetPhysicsScene(), FInitBodySpawnParams(IsStaticPhysics(), /*bPhysicsTypeDeterminesSimulation*/false), &BodyInstanceOwner);
+			BodyInstance.InitBody(BodySetup, BodyTransform, SurrogateComponent, GetWorld()->GetPhysicsScene(), FInitBodySpawnParams(IsStaticPhysics(), /*bPhysicsTypeDeterminesSimulation*/false), &BodyInstanceOwner);
 
 			// Assign BodyInstanceOwner
 			if (Chaos::FSingleParticlePhysicsProxy* ProxyHandle = BodyInstance.GetPhysicsActor())
@@ -889,5 +905,44 @@ UPhysicalMaterial* FFastGeoPrimitiveComponent::GetPhysicsMaterialOverride() cons
 
 FBodyInstance* FFastGeoPrimitiveComponent::GetBodyInstance(FName BoneName, bool bGetWelded, int32 Index) const
 {
+	if (Index != INDEX_NONE)
+	{
+		// Handle the case where GetBodyInstance is called directly on the FastGeo component with an already transformed Index
+		if (UFastGeoSurrogateComponent* SurrogateComponent = GetSurrogateComponent())
+		{
+			if (FFastGeoSurrogateBodyInstanceIndex::IsEncoded(Index))
+			{
+				return SurrogateComponent->GetBodyInstance(BoneName, bGetWelded, Index);
+			}
+		}
+	}
 	return const_cast<FBodyInstance*>(&BodyInstance);
-}
\ No newline at end of file
+}
+
+UFastGeoSurrogateComponent* FFastGeoPrimitiveComponent::GetSurrogateComponent() const
+{
+	if (UFastGeoContainer* Container = GetOwnerContainer(); ensure(Container))
+	{
+		if (Container->IsUsingSurrogateComponents())
+		{
+			if (const AFastGeoSurrogateActor* SurrogateActor = Container->GetSurrogateActor())
+			{
+				UFastGeoSurrogateComponent* Component = SurrogateActor->GetSurrogateComponent(SurrogateComponentDescriptorIndex);
+				ensure(Component);
+				return Component;
+			}
+		}
+	}
+	return nullptr;
+}
+
+#if WITH_EDITOR
+void FFastGeoPrimitiveComponent::ReserveSurrogateInstanceBodyIndices(FFastGeoSurrogateBodyInstanceIndex& InOutNextAvailableInstanceBodyIndex)
+{
+	check(BodyInstance.InstanceBodyIndex == INDEX_NONE);
+	BodyInstance.InstanceBodyIndex = InOutNextAvailableInstanceBodyIndex.GetEncoded();
+	check(FFastGeoSurrogateBodyInstanceIndex::IsEncoded(BodyInstance.InstanceBodyIndex));
+
+	++InOutNextAvailableInstanceBodyIndex;
+}
+#endif
\ No newline at end of file
diff --git a/Engine/Plugins/Experimental/FastGeoStreaming/Source/FastGeoStreaming/Private/FastGeoStaticMeshComponent.cpp b/Engine/Plugins/Experimental/FastGeoStreaming/Source/FastGeoStreaming/Private/FastGeoStaticMeshComponent.cpp
--- a/Engine/Plugins/Experimental/FastGeoStreaming/Source/FastGeoStreaming/Private/FastGeoStaticMeshComponent.cpp
+++ b/Engine/Plugins/Experimental/FastGeoStreaming/Source/FastGeoStreaming/Private/FastGeoStaticMeshComponent.cpp
@@ -477,3 +477,8 @@ FFastGeoStaticMeshComponent::FFastGeoStaticMeshComponent(int32 InComponentIndex,
 	: Super(InComponentIndex, InType)
 {
 }
+
+void FFastGeoStaticMeshComponent::GetBodyInstances(TArray<FBodyInstance*>& OutBodyInstances)
+{
+	OutBodyInstances.Add(&BodyInstance);
+}
\ No newline at end of file
diff --git a/Engine/Plugins/Experimental/FastGeoStreaming/Source/FastGeoStreaming/Private/FastGeoSurrogateActor.cpp b/Engine/Plugins/Experimental/FastGeoStreaming/Source/FastGeoStreaming/Private/FastGeoSurrogateActor.cpp
new file mode 100644
--- /dev/null
+++ b/Engine/Plugins/Experimental/FastGeoStreaming/Source/FastGeoStreaming/Private/FastGeoSurrogateActor.cpp
@@ -0,0 +1,136 @@
+// Copyright Epic Games, Inc. All Rights Reserved.
+
+#include "FastGeoSurrogateActor.h"
+#include "FastGeoSurrogateComponent.h"
+#include "FastGeoWorldSubsystem.h"
+#include "FastGeoSurrogateComponentDescriptor.h"
+#include "FastGeoContainer.h"
+#include "Engine/Level.h"
+#include "Engine/World.h"
+
+AFastGeoSurrogateActor::AFastGeoSurrogateActor()
+{
+	PrimaryActorTick.bCanEverTick = false;
+	bReplicates = false;
+	RootComponent = CreateDefaultSubobject<USceneComponent>(TEXT("FastGeoCellProxyRootComponent"));
+	RootComponent->SetMobility(EComponentMobility::Static);
+	SetRootComponent(RootComponent);
+}
+
+void AFastGeoSurrogateActor::PostInitializeComponents()
+{
+	TRACE_CPUPROFILER_EVENT_SCOPE(AFastGeoSurrogateActor::PostInitializeComponents);
+	Super::PostInitializeComponents();
+
+	if (!UFastGeoWorldSubsystem::ShouldAllowSurrogateComponents())
+	{
+		return;
+	}
+
+	bIsActive = true;
+
+	// Prepare SurrogateComponents
+	TArray<UFastGeoSurrogateComponent*> Components;
+	GetComponents<UFastGeoSurrogateComponent>(Components);
+	SurrogateComponents.SetNumZeroed(Components.Num());
+	for (UFastGeoSurrogateComponent* Component : Components)
+	{
+		// Descriptor count is equal to surrogate component count, so descriptor index should always be valid
+		if (ensure(SurrogateComponents.IsValidIndex(Component->GetDescriptorIndex())))
+		{
+			SurrogateComponents[Component->GetDescriptorIndex()] = Component;
+		}
+	}
+
+	// Register FastGeo content
+	check(!bHasRegisteredWithFastGeo);
+	if (!bHasRegisteredWithFastGeo && GetWorld() && GetWorld()->IsGameWorld())
+	{
+		ULevel* Level = GetLevel();
+		if (UFastGeoContainer* FastGeo = Level ? Level->GetAssetUserData<UFastGeoContainer>() : nullptr; ensure(FastGeo))
+		{
+			FastGeo->Register();
+			bHasRegisteredWithFastGeo = true;
+		}
+	}
+}
+
+void AFastGeoSurrogateActor::PostUnregisterAllComponents()
+{
+	Super::PostUnregisterAllComponents();
+
+	UnregisterFromFastGeo();
+}
+
+void AFastGeoSurrogateActor::EndPlay(const EEndPlayReason::Type EndPlayReason)
+{
+	// Normally, PostUnregisterAllComponents should have been called, but NOT in the case 
+	// where RemoveFromWorld was called with bAllowIncrementalRemoval = false.
+	// In this case, we need to wait for FastGeo to complete.
+	if (UnregisterFromFastGeo())
+	{
+		ULevel* Level = GetLevel();
+		if (UFastGeoContainer* FastGeo = Level ? Level->GetAssetUserData<UFastGeoContainer>() : nullptr; ensure(FastGeo))
+		{
+			const bool bWaitForCompletion = true;
+			FastGeo->Tick(bWaitForCompletion);
+		}
+	}
+
+	Super::EndPlay(EndPlayReason);
+}
+
+bool AFastGeoSurrogateActor::UnregisterFromFastGeo()
+{
+	if (!UFastGeoWorldSubsystem::ShouldAllowSurrogateComponents())
+	{
+		return false;
+	}
+
+	TRACE_CPUPROFILER_EVENT_SCOPE(AFastGeoSurrogateActor::UnregisterFromFastGeo);
+	bIsActive = false;
+
+	if (bHasRegisteredWithFastGeo)
+	{
+		ULevel* Level = GetLevel();
+		if (UFastGeoContainer* FastGeo = Level ? Level->GetAssetUserData<UFastGeoContainer>() : nullptr; ensure(FastGeo))
+		{
+			check(FastGeo->GetSurrogateActor() == this);
+
+			// Mutate body instances to make sure any direct access to the bodyinstance will return the intended value.
+			// Direct BodyInstance access is safe because it only changes local CPU data, no physics-thread sync is required.
+			ForEachComponent<UFastGeoSurrogateComponent>(false, [](UFastGeoSurrogateComponent* SurrogateComponent)
+			{
+				SurrogateComponent->GetBodyInstance()->SetCollisionEnabled(ECollisionEnabled::NoCollision, false);
+				SurrogateComponent->GetBodyInstance()->SetResponseToAllChannels(ECR_Ignore);
+			});
+
+			FastGeo->Unregister();
+			bHasRegisteredWithFastGeo = false;
+			return true;
+		}
+	}
+	return false;
+}
+
+UFastGeoSurrogateComponent* AFastGeoSurrogateActor::GetSurrogateComponent(int32 SurrogateComponentDescriptorIndex) const
+{
+	if (!UFastGeoWorldSubsystem::ShouldAllowSurrogateComponents())
+	{ 
+		return nullptr;
+	}
+	if (ensure(SurrogateComponents.IsValidIndex(SurrogateComponentDescriptorIndex)))
+	{
+		return SurrogateComponents[SurrogateComponentDescriptorIndex];
+	}
+	return nullptr;
+}
+
+const FFastGeoSurrogateComponentDescriptor* AFastGeoSurrogateActor::GetSurrogateComponentDescriptor(int32 SurrogateComponentDescriptorIndex) const
+{
+	if (ensure(SurrogateComponentDescriptors.IsValidIndex(SurrogateComponentDescriptorIndex)))
+	{
+		return &SurrogateComponentDescriptors[SurrogateComponentDescriptorIndex];
+	}
+	return nullptr;
+}
\ No newline at end of file
diff --git a/Engine/Plugins/Experimental/FastGeoStreaming/Source/FastGeoStreaming/Private/FastGeoSurrogateComponent.cpp b/Engine/Plugins/Experimental/FastGeoStreaming/Source/FastGeoStreaming/Private/FastGeoSurrogateComponent.cpp
new file mode 100644
--- /dev/null
+++ b/Engine/Plugins/Experimental/FastGeoStreaming/Source/FastGeoStreaming/Private/FastGeoSurrogateComponent.cpp
@@ -0,0 +1,168 @@
+// Copyright Epic Games, Inc. All Rights Reserved.
+
+#include "FastGeoSurrogateComponent.h"
+#include "FastGeoSurrogateActor.h"
+#include "FastGeoContainer.h"
+#include "FastGeoComponent.h"
+#include "FastGeoPrimitiveComponent.h"
+#include "FastGeoWorldSubsystem.h"
+#include "FastGeoSurrogateBodyInstanceIndex.h"
+#include "FastGeoSurrogateComponentDescriptor.h"
+#include "WorldPartition/WorldPartitionRuntimeCellInterface.h"
+#include "GameFramework/Actor.h"
+#include "Engine/Level.h"
+
+UFastGeoSurrogateComponent::UFastGeoSurrogateComponent()
+{
+	PrimaryComponentTick.bCanEverTick = false;
+	SetHiddenInGame(true);
+	SetCollisionEnabled(ECollisionEnabled::QueryAndPhysics);
+	SetCollisionObjectType(ECC_WorldStatic);
+	SetCollisionResponseToAllChannels(ECR_Block);
+	SetGenerateOverlapEvents(false);
+	SetCanEverAffectNavigation(false);
+	bAffectDynamicIndirectLighting = false;
+	bAffectDistanceFieldLighting = false;
+	bNeverNeedsRenderUpdate = true;
+	SetMobility(EComponentMobility::Static);
+}
+
+void UFastGeoSurrogateComponent::OnRegister()
+{
+	if (AFastGeoSurrogateActor* SurrogateActor = Cast<AFastGeoSurrogateActor>(GetOwner()); ensure(SurrogateActor))
+	{
+		if (const FFastGeoSurrogateComponentDescriptor* ComponentDesc = SurrogateActor->GetSurrogateComponentDescriptor(DescriptorIndex); ensure(ComponentDesc))
+		{
+			SetSurrogateComponentDescriptor(*ComponentDesc);
+		}
+	}
+
+	Super::OnRegister();
+}
+
+void UFastGeoSurrogateComponent::SetSurrogateComponentDescriptor(const FFastGeoSurrogateComponentDescriptor& InComponentDesc)
+{
+	check(!IsPhysicsStateCreated());
+	BodyInstance.SetCollisionEnabled(InComponentDesc.CollisionEnabled);
+	BodyInstance.SetObjectType(InComponentDesc.ObjectType);
+	BodyInstance.SetResponseToChannels(InComponentDesc.ResponseContainer);
+	BodyInstance.SetWalkableSlopeOverride(InComponentDesc.WalkableSlopeOverride, true);
+	CanCharacterStepUpOn = InComponentDesc.CanCharacterStepUpOn;
+}
+
+FBoxSphereBounds UFastGeoSurrogateComponent::CalcBounds(const FTransform& LocalToWorld) const
+{
+	AActor* Owner = GetOwner();
+	ULevel* Level = Owner ? Owner->GetLevel() : nullptr;
+	if (const IWorldPartitionCell* Cell = Level ? Level->GetWorldPartitionRuntimeCell() : nullptr)
+	{
+		return Cell->GetContentBounds();
+	}
+
+	const FVector Origin = LocalToWorld.GetLocation();
+	return FBoxSphereBounds(Origin, FVector(1.f), 1.f);
+}
+
+bool UFastGeoSurrogateComponent::ShouldCreatePhysicsState() const
+{
+	if (!UFastGeoWorldSubsystem::ShouldAllowSurrogateComponents())
+	{
+		return false;
+	}
+	return Super::ShouldCreatePhysicsState();
+}
+
+void UFastGeoSurrogateComponent::OnCreatePhysicsState()
+{
+	// Skip PrimitiveComponent implementation
+	USceneComponent::OnCreatePhysicsState();
+
+	if (!IsAsyncCreatePhysicsStateRunning())
+	{
+		OnComponentPhysicsStateChanged.Broadcast(this, EComponentPhysicsStateChange::Created);
+	}
+}
+
+ECollisionEnabled::Type UFastGeoSurrogateComponent::GetCollisionEnabled() const
+{
+	if (AFastGeoSurrogateActor* SurrogateOwner = Cast<AFastGeoSurrogateActor>(GetOwner()))
+	{
+		if (!SurrogateOwner->IsActive())
+		{
+			return ECollisionEnabled::NoCollision;
+		}
+	}
+	return Super::GetCollisionEnabled();
+}
+
+ECollisionResponse UFastGeoSurrogateComponent::GetCollisionResponseToChannel(ECollisionChannel Channel) const
+{
+	if (AFastGeoSurrogateActor* SurrogateOwner = Cast<AFastGeoSurrogateActor>(GetOwner()))
+	{
+		if (!SurrogateOwner->IsActive())
+		{
+			return ECollisionResponse::ECR_Ignore;
+		}
+	}
+	return Super::GetCollisionResponseToChannel(Channel);
+}
+
+const FCollisionResponseContainer& UFastGeoSurrogateComponent::GetCollisionResponseToChannels() const
+{
+	if (AFastGeoSurrogateActor* SurrogateOwner = Cast<AFastGeoSurrogateActor>(GetOwner()))
+	{
+		if (!SurrogateOwner->IsActive())
+		{
+			static const FCollisionResponseContainer AllIgnoreContainer = []() 
+			{
+				FCollisionResponseContainer Container;
+				Container.SetAllChannels(ECollisionResponse::ECR_Ignore);
+				return Container;
+			}();
+
+			return AllIgnoreContainer;
+		}
+	}
+	return Super::GetCollisionResponseToChannels();
+}
+
+void UFastGeoSurrogateComponent::RegisterBodyInstances(const TArray<FBodyInstance*>& InBodyInstances)
+{
+	TRACE_CPUPROFILER_EVENT_SCOPE(UFastGeoSurrogateComponent::RegisterBodyInstances);
+	check(RegisteredBodyInstanceCount == 0);
+	check(IndexToBodyInstance.IsEmpty());
+	IndexToBodyInstance.SetNumZeroed(InBodyInstances.Num());
+	for (FBodyInstance* BodyInst : InBodyInstances)
+	{
+		if (BodyInst)
+		{
+			check(FFastGeoSurrogateBodyInstanceIndex::IsEncoded(BodyInst->InstanceBodyIndex));
+			const int32 ArrayIndex = FFastGeoSurrogateBodyInstanceIndex::Decode(BodyInst->InstanceBodyIndex);
+			check(IndexToBodyInstance.IsValidIndex(ArrayIndex));
+			check(!IndexToBodyInstance[ArrayIndex]);
+			IndexToBodyInstance[ArrayIndex] = BodyInst;
+		}
+	}
+}
+
+void UFastGeoSurrogateComponent::UnregisterBodyInstances()
+{
+	TRACE_CPUPROFILER_EVENT_SCOPE(UFastGeoSurrogateComponent::UnregisterBodyInstances);
+	RegisteredBodyInstanceCount = 0;
+	IndexToBodyInstance.Reset();
+}
+
+FBodyInstance* UFastGeoSurrogateComponent::GetBodyInstance(FName BoneName, bool bGetWelded, int32 Index) const
+{
+	if (Index != INDEX_NONE && FFastGeoSurrogateBodyInstanceIndex::IsEncoded(Index))
+	{
+		const int32 LocalIndex = FFastGeoSurrogateBodyInstanceIndex::Decode(Index);
+		if (IndexToBodyInstance.IsValidIndex(LocalIndex))
+		{
+			check(IndexToBodyInstance[LocalIndex]);
+			return IndexToBodyInstance[LocalIndex];
+		}
+		return nullptr;
+	}
+	return Super::GetBodyInstance(BoneName, bGetWelded, Index);
+}
\ No newline at end of file
diff --git a/Engine/Plugins/Experimental/FastGeoStreaming/Source/FastGeoStreaming/Private/FastGeoWorldPartitionRuntimeCellTransformer.cpp b/Engine/Plugins/Experimental/FastGeoStreaming/Source/FastGeoStreaming/Private/FastGeoWorldPartitionRuntimeCellTransformer.cpp
--- a/Engine/Plugins/Experimental/FastGeoStreaming/Source/FastGeoStreaming/Private/FastGeoWorldPartitionRuntimeCellTransformer.cpp
+++ b/Engine/Plugins/Experimental/FastGeoStreaming/Source/FastGeoStreaming/Private/FastGeoWorldPartitionRuntimeCellTransformer.cpp
@@ -44,6 +44,10 @@
 #include "FastGeoInstancedStaticMeshComponent.h"
 #include "FastGeoSkinnedMeshComponent.h"
 #include "FastGeoInstancedSkinnedMeshComponent.h"
+#include "FastGeoSurrogateActor.h"
+#include "FastGeoSurrogateComponent.h"
+#include "FastGeoSurrogateComponentDescriptor.h"
+#include "FastGeoSurrogateBodyInstanceIndex.h"
 #include "FastGeoLog.h"
 #endif
 
@@ -355,6 +359,8 @@ void UFastGeoWorldPartitionRuntimeCellTransformer::Transform(ULevel* InLevel)
 
 		TUniquePtr<FFastGeoComponentCluster> LevelComponentCluster(new FFastGeoComponentCluster(FastGeo, *FString::Printf(TEXT("FastGeoComponentCluster_%s"), *CellName)));
 
+		TMap<int32, FFastGeoSurrogateBodyInstanceIndex> LevelSurrogateComponentBodyInstanceCounter;
+		TArray<FFastGeoSurrogateComponentDescriptor> LevelSurrogateComponentDescriptors;
 		for (TMap<AActor*, FTransformableActor>::TConstIterator It = TransformableActors.CreateConstIterator(); It; ++It)
 		{
 			const TPair<AActor*, FTransformableActor>& Entry = *It;
@@ -391,6 +397,18 @@ void UFastGeoWorldPartitionRuntimeCellTransformer::Transform(ULevel* InLevel)
 				FFastGeoComponent& FastGeoComponent = CurrentComponentCluster->AddComponent(FastGeoComponentType);
 				FastGeoComponent.InitializeFromComponent(Component);
 
+				if (GetSettings().bGenerateSurrogateComponents)
+				{
+					if (FFastGeoPrimitiveComponent* FastGeoPrimitiveComponent = FastGeoComponent.CastTo<FFastGeoPrimitiveComponent>(); FastGeoPrimitiveComponent && FastGeoPrimitiveComponent->IsCollisionEnabled())
+					{
+						UPrimitiveComponent* PrimitiveComponent = CastChecked<UPrimitiveComponent>(Component);
+						FFastGeoSurrogateComponentDescriptor ComponentDesc(*PrimitiveComponent);
+						FastGeoPrimitiveComponent->SurrogateComponentDescriptorIndex = LevelSurrogateComponentDescriptors.AddUnique(ComponentDesc);
+						FFastGeoSurrogateBodyInstanceIndex& NextAvailableInstanceBodyIndex = LevelSurrogateComponentBodyInstanceCounter.FindOrAdd(FastGeoPrimitiveComponent->SurrogateComponentDescriptorIndex, FFastGeoSurrogateBodyInstanceIndex::Encode(0));
+						FastGeoPrimitiveComponent->ReserveSurrogateInstanceBodyIndices(NextAvailableInstanceBodyIndex);
+					}
+				}
+
 				// Remove the component from the actor
 				Actor->RemoveOwnedComponent(Component);
 				Component->MarkAsGarbage();
@@ -431,6 +449,30 @@ void UFastGeoWorldPartitionRuntimeCellTransformer::Transform(ULevel* InLevel)
 			FastGeo->AddComponentCluster(LevelComponentCluster.Get());
 		}
 
+		if (LevelSurrogateComponentDescriptors.Num() > 0)
+		{
+			AFastGeoSurrogateActor* SurrogateActor = NewObject<AFastGeoSurrogateActor>(InLevel, AFastGeoSurrogateActor::StaticClass());
+			if (const IWorldPartitionCell* Cell = InLevel->GetWorldPartitionRuntimeCell())
+			{
+				SurrogateActor->SetActorLocation(Cell->GetCellBounds().GetCenter());
+			}
+			SurrogateActor->SurrogateComponentDescriptors = MoveTemp(LevelSurrogateComponentDescriptors);
+			InLevel->Actors.Add(SurrogateActor);
+			check(SurrogateActor->GetLevel() == InLevel);
+
+			for (int32 DescriptorIndex = 0; DescriptorIndex < SurrogateActor->SurrogateComponentDescriptors.Num(); ++DescriptorIndex)
+			{
+				const FFastGeoSurrogateComponentDescriptor& ComponentDesc = SurrogateActor->SurrogateComponentDescriptors[DescriptorIndex];
+				UFastGeoSurrogateComponent* SurrogateComp = NewObject<UFastGeoSurrogateComponent>(SurrogateActor);
+				SurrogateActor->AddInstanceComponent(SurrogateComp);
+				SurrogateComp->DescriptorIndex = DescriptorIndex;
+				SurrogateComp->SetSurrogateComponentDescriptor(ComponentDesc);
+				SurrogateComp->SetupAttachment(SurrogateActor->GetRootComponent());
+			}
+
+			FastGeo->SurrogateActor = SurrogateActor;
+		}
+
 		// Finalize post-creation intialization
 		FastGeo->OnCreated();
 	}
diff --git a/Engine/Plugins/Experimental/FastGeoStreaming/Source/FastGeoStreaming/Private/FastGeoWorldSubsystem.cpp b/Engine/Plugins/Experimental/FastGeoStreaming/Source/FastGeoStreaming/Private/FastGeoWorldSubsystem.cpp
--- a/Engine/Plugins/Experimental/FastGeoStreaming/Source/FastGeoStreaming/Private/FastGeoWorldSubsystem.cpp
+++ b/Engine/Plugins/Experimental/FastGeoStreaming/Source/FastGeoStreaming/Private/FastGeoWorldSubsystem.cpp
@@ -4,6 +4,7 @@
 #include "FastGeoContainer.h"
 #include "FastGeoPrimitiveComponent.h"
 #include "FastGeoStreamingModule.h"
+#include "FastGeoSurrogateActor.h"
 #include "Physics/Experimental/PhysScene_Chaos.h"
 #include "Streaming/LevelStreamingDelegates.h"
 #include "ComponentReregisterContext.h"
@@ -42,6 +43,15 @@ namespace FastGeo
 		TEXT("Maximum number of components to process (0 = no component limit)"),
 		ECVF_Default);
 
+	static bool GAllowSurrogateComponents = true;
+	static FAutoConsoleVariableRef CVarAllowSurrogateComponents(
+		TEXT("FastGeo.AllowSurrogateComponents"),
+		GAllowSurrogateComponents,
+		TEXT("Allows the use of generated FastGeo surrogate components.\n")
+		TEXT("Must be set at startup; cannot be changed during runtime.\n")
+		TEXT("Requires data transformed with 'bGenerateSurrogateComponents' enabled."),
+		ECVF_ReadOnly);
+
 #if !UE_BUILD_SHIPPING
 	static bool GShowFastGeo = true;
 	FAutoConsoleCommand GShowFastGeoCommand(
@@ -86,6 +96,11 @@ bool UFastGeoWorldSubsystem::IsFastGeoVisible()
 }
 #endif
 
+bool UFastGeoWorldSubsystem::ShouldAllowSurrogateComponents()
+{
+	return FastGeo::GAllowSurrogateComponents;
+}
+
 UFastGeoWorldSubsystem::UFastGeoWorldSubsystem()
 {
 	if (!FFastGeoStreamingModule::IsFastGeoEnabled())
@@ -267,9 +282,13 @@ void UFastGeoWorldSubsystem::OnLevelStartedAddToWorld(UWorld* World, const ULeve
 
 	if (UFastGeoContainer* FastGeo = Level->GetAssetUserData<UFastGeoContainer>())
 	{
-		TRACE_CPUPROFILER_EVENT_SCOPE(UFastGeoWorldSubsystem::OnLevelStartedAddToWorld);
+		// Registration is handled by surrogate actor
+		if (!FastGeo->IsUsingSurrogateComponents())
+		{
+			TRACE_CPUPROFILER_EVENT_SCOPE(UFastGeoWorldSubsystem::OnLevelStartedAddToWorld);
 
-		FastGeo->Register();
+			FastGeo->Register();
+		}
 	}
 }
 
@@ -282,9 +301,13 @@ void UFastGeoWorldSubsystem::OnLevelStartedRemoveFromWorld(UWorld* World, const
 
 	if (UFastGeoContainer* FastGeo = Level->GetAssetUserData<UFastGeoContainer>())
 	{
-		TRACE_CPUPROFILER_EVENT_SCOPE(UFastGeoWorldSubsystem::OnLevelStartedRemoveFromWorld);
+		// Unregistration is handled by surrogate actor
+		if (!FastGeo->IsUsingSurrogateComponents())
+		{
+			TRACE_CPUPROFILER_EVENT_SCOPE(UFastGeoWorldSubsystem::OnLevelStartedRemoveFromWorld);
 
-		FastGeo->Unregister();
+			FastGeo->Unregister();
+		}
 	}
 }
 
diff --git a/Engine/Plugins/Experimental/FastGeoStreaming/Source/FastGeoStreaming/Public/FastGeoWorldPartitionRuntimeCellTransformer.h b/Engine/Plugins/Experimental/FastGeoStreaming/Source/FastGeoStreaming/Public/FastGeoWorldPartitionRuntimeCellTransformer.h
--- a/Engine/Plugins/Experimental/FastGeoStreaming/Source/FastGeoStreaming/Public/FastGeoWorldPartitionRuntimeCellTransformer.h
+++ b/Engine/Plugins/Experimental/FastGeoStreaming/Source/FastGeoStreaming/Public/FastGeoWorldPartitionRuntimeCellTransformer.h
@@ -139,6 +139,15 @@ public:
 	UPROPERTY(EditAnywhere, Category="Settings", meta = (AllowAbstract = "true"))
 	TArray<TSubclassOf<UActorComponent>> IgnoredRemainingExactComponentClasses;
 
+	/**
+	  * Whether the transformer should generate a surrogate actor and surrogate components
+	  * which will be linked to the physics body instances when initialized.
+	  * When hitting FastGeo content, a surrogate component will be available in the HitResult & OverlapResult.
+	  * HitResult.Item & OverlapResult.ItemIndex will map to the actual FastGeo component body instance index.
+	  */
+	UPROPERTY(EditAnywhere, Category = "Settings")
+	bool bGenerateSurrogateComponents = true;
+
 	/** Determines how instance transforms are stored in memory */
 	UPROPERTY(EditAnywhere, Category = "Settings")
 	EFastGeoInstanceStorageMode InstanceStorageMode = EFastGeoInstanceStorageMode::Full;
diff --git a/Engine/Source/Runtime/Engine/Classes/PhysicsEngine/PhysicsBodyInstanceOwnerInterface.h b/Engine/Source/Runtime/Engine/Classes/PhysicsEngine/PhysicsBodyInstanceOwnerInterface.h
--- a/Engine/Source/Runtime/Engine/Classes/PhysicsEngine/PhysicsBodyInstanceOwnerInterface.h
+++ b/Engine/Source/Runtime/Engine/Classes/PhysicsEngine/PhysicsBodyInstanceOwnerInterface.h
@@ -52,11 +52,17 @@ public:
 		return const_cast<IPhysicsBodyInstanceOwner*>(GetPhysicsBodyInstanceOwnerFromOverlapResult(OverlapResult));
 	}
 
-	/** Returns the IPhysicsBodyInstanceOwner based on a given hit result. */
-	static ENGINE_API const IPhysicsBodyInstanceOwner* GetPhysicsBodyInstanceOwnerFromHitResult(const FHitResult& Result);
+	/** Returns the IPhysicsBodyInstanceOwner based on a given hit result. 
+	 * @param HiResult Hit result
+	 * @param bPreferComponent Whether to prioritize the hit result component (if valid) when resolving the IPhysicsBodyInstanceOwner.
+	 */
+	static ENGINE_API const IPhysicsBodyInstanceOwner* GetPhysicsBodyInstanceOwnerFromHitResult(const FHitResult& HiResult, bool bPreferComponent = true);
 
-	/** Returns the IPhysicsBodyInstanceOwner based on a given overlap result. */
-	static ENGINE_API const IPhysicsBodyInstanceOwner* GetPhysicsBodyInstanceOwnerFromOverlapResult(const FOverlapResult& OverlapResult);
+	/** Returns the IPhysicsBodyInstanceOwner based on a given overlap result.
+	 * @param OverlapResult Overlap result
+	 * @param bPreferComponent Whether to prioritize the overlap result component (if valid) when resolving the IPhysicsBodyInstanceOwner.
+	 */
+	static ENGINE_API const IPhysicsBodyInstanceOwner* GetPhysicsBodyInstanceOwnerFromOverlapResult(const FOverlapResult& OverlapResult, bool bPreferComponent = true);
 
 	/** Whether the physics is static. */
 	virtual bool IsStaticPhysics() const = 0;
@@ -112,7 +118,7 @@ public:
 	/** Returns the body instance owner number of materials (used by FBodyInstance::GetComplexPhysicalMaterials). */
 	virtual int32 GetNumMaterials() const = 0;
 
-	/** Returns the body instance owner material for the provided material index. (used by FBodyInstance::GetComplexPhysicalMaterials). */
+	/** Returns the body instance owner material for the provided material index (used by FBodyInstance::GetComplexPhysicalMaterials). */
 	virtual UMaterialInterface* GetMaterial(int32 Index) const = 0;
 
 	friend struct FBodyInstance;
diff --git a/Engine/Source/Runtime/Engine/Private/PhysicsEngine/PhysicsBodyInstanceOwnerInterface.cpp b/Engine/Source/Runtime/Engine/Private/PhysicsEngine/PhysicsBodyInstanceOwnerInterface.cpp
--- a/Engine/Source/Runtime/Engine/Private/PhysicsEngine/PhysicsBodyInstanceOwnerInterface.cpp
+++ b/Engine/Source/Runtime/Engine/Private/PhysicsEngine/PhysicsBodyInstanceOwnerInterface.cpp
@@ -7,8 +7,16 @@
 
 #include UE_INLINE_GENERATED_CPP_BY_NAME(PhysicsBodyInstanceOwnerInterface)
 
-const IPhysicsBodyInstanceOwner* IPhysicsBodyInstanceOwner::GetPhysicsBodyInstanceOwnerFromHitResult(const FHitResult& Result)
+const IPhysicsBodyInstanceOwner* IPhysicsBodyInstanceOwner::GetPhysicsBodyInstanceOwnerFromHitResult(const FHitResult& Result, bool bPreferComponent)
 {
+	if (bPreferComponent)
+	{
+		if (const UPrimitiveComponent* Component = Result.GetComponent())
+		{
+			return Component;
+		}
+	}
+
 	if (IPhysicsBodyInstanceOwnerResolver* PhysicsObjectOwnerResolver = Result.PhysicsObject ? Cast<IPhysicsBodyInstanceOwnerResolver>(Result.PhysicsObjectOwner.Get()) : nullptr)
 	{
 		if (IPhysicsBodyInstanceOwner* BodyInstanceOwner = PhysicsObjectOwnerResolver->ResolvePhysicsBodyInstanceOwner(Result.PhysicsObject))
@@ -16,15 +24,20 @@ const IPhysicsBodyInstanceOwner* IPhysicsBodyInstanceOwner::GetPhysicsBodyInstan
 			return BodyInstanceOwner;
 		}
 	}
-	if (const UPrimitiveComponent* Component = Result.GetComponent())
-	{
-		return Component;
-	}
-	return nullptr;
+
+	return !bPreferComponent ? Result.GetComponent() : nullptr;
 }
 
-const IPhysicsBodyInstanceOwner* IPhysicsBodyInstanceOwner::GetPhysicsBodyInstanceOwnerFromOverlapResult(const FOverlapResult& Result)
+const IPhysicsBodyInstanceOwner* IPhysicsBodyInstanceOwner::GetPhysicsBodyInstanceOwnerFromOverlapResult(const FOverlapResult& Result, bool bPreferComponent)
 {
+	if (bPreferComponent)
+	{
+		if (const UPrimitiveComponent* Component = Result.GetComponent())
+		{
+			return Component;
+		}
+	}
+
 	if (IPhysicsBodyInstanceOwnerResolver* PhysicsObjectOwnerResolver = Result.PhysicsObject ? Cast<IPhysicsBodyInstanceOwnerResolver>(Result.PhysicsObjectOwner.Get()) : nullptr)
 	{
 		if (IPhysicsBodyInstanceOwner* BodyInstanceOwner = PhysicsObjectOwnerResolver->ResolvePhysicsBodyInstanceOwner(Result.PhysicsObject))
@@ -32,9 +45,6 @@ const IPhysicsBodyInstanceOwner* IPhysicsBodyInstanceOwner::GetPhysicsBodyInstan
 			return BodyInstanceOwner;
 		}
 	}
-	if (const UPrimitiveComponent* Component = Result.GetComponent())
-	{
-		return Component;
-	}
-	return nullptr;
+
+	return !bPreferComponent ? Result.GetComponent() : nullptr;
 }
