[FastGeo] FastGeo transformer now doesn't allow to transform primitives with flag 'Simulation Generates Hit Events' set.

#Synced-CL 50298075

diff --git a/Engine/Plugins/Experimental/FastGeoStreaming/Source/FastGeoStreaming/Private/FastGeoWorldPartitionRuntimeCellTransformer.cpp b/Engine/Plugins/Experimental/FastGeoStreaming/Source/FastGeoStreaming/Private/FastGeoWorldPartitionRuntimeCellTransformer.cpp
--- a/Engine/Plugins/Experimental/FastGeoStreaming/Source/FastGeoStreaming/Private/FastGeoWorldPartitionRuntimeCellTransformer.cpp
+++ b/Engine/Plugins/Experimental/FastGeoStreaming/Source/FastGeoStreaming/Private/FastGeoWorldPartitionRuntimeCellTransformer.cpp
@@ -993,6 +993,11 @@ FFastGeoTransformResult UFastGeoWorldPartitionRuntimeCellTransformer::CanTransfo
 		return { EFastGeoTransform::Reject, [PrimitiveComponent] { return FString::Printf(TEXT("Component %s has a valid LOD Parent Primitive"), *FastGeo::GetComponentShortName(PrimitiveComponent)); } };
 	}
 
+	if (PrimitiveComponent && PrimitiveComponent->BodyInstance.bNotifyRigidBodyCollision)
+	{
+		return { EFastGeoTransform::Reject, [PrimitiveComponent] { return FString::Printf(TEXT("Component %s has the 'Simulation Generates Hit Events' flag set"), *FastGeo::GetComponentShortName(PrimitiveComponent)); }};
+	}
+
 	bool bIsSceneComponentVisible = SceneComponent->IsVisible() && !SceneComponent->GetOwner()->IsHidden();
 	bool bIsPrimitiveVisible = PrimitiveComponent && (PrimitiveComponent->bCastHiddenShadow || PrimitiveComponent->bAffectIndirectLightingWhileHidden || PrimitiveComponent->bRayTracingFarField);
 	bool bShouldAddToRenderScene = bIsSceneComponentVisible || bIsPrimitiveVisible;
