[FastGeo] Sort transformed actors/components to ensure deterministic cook output (also important to generate deterministic surrogate body instance indices).

#Synced-CL 50298075

diff --git a/Engine/Plugins/Experimental/FastGeoStreaming/Source/FastGeoStreaming/Private/FastGeoWorldPartitionRuntimeCellTransformer.cpp b/Engine/Plugins/Experimental/FastGeoStreaming/Source/FastGeoStreaming/Private/FastGeoWorldPartitionRuntimeCellTransformer.cpp
--- a/Engine/Plugins/Experimental/FastGeoStreaming/Source/FastGeoStreaming/Private/FastGeoWorldPartitionRuntimeCellTransformer.cpp
+++ b/Engine/Plugins/Experimental/FastGeoStreaming/Source/FastGeoStreaming/Private/FastGeoWorldPartitionRuntimeCellTransformer.cpp
@@ -295,7 +295,7 @@ void UFastGeoWorldPartitionRuntimeCellTransformer::OnSelectionChanged(UObject* O
 			UE_LOG(LogFastGeoStreaming, Log, TEXT("- FastGeoStreaming Debug Mode: Transformation on %d selected actors "), SelectedActors.Num());
 
 			FTransformationStats Stats;
-			TMap<AActor*, FTransformableActor> TransformableActors;
+			TArray<TPair<AActor*, FTransformableActor>> TransformableActors;
 
 			GatherTransformableActors(SelectedActors, SelectedActors[0]->GetLevel(), TransformableActors, Stats);
 
@@ -348,7 +348,7 @@ void UFastGeoWorldPartitionRuntimeCellTransformer::Transform(ULevel* InLevel)
 	}
 
 	FTransformationStats Stats;
-	TMap<AActor*, FTransformableActor> TransformableActors;
+	TArray<TPair<AActor*, FTransformableActor>> TransformableActors;
 	GatherTransformableActors(InLevel->Actors, InLevel, TransformableActors, Stats);
 
 	if (!TransformableActors.IsEmpty())
@@ -361,9 +361,9 @@ void UFastGeoWorldPartitionRuntimeCellTransformer::Transform(ULevel* InLevel)
 
 		TMap<int32, FFastGeoSurrogateBodyInstanceIndex> LevelSurrogateComponentBodyInstanceCounter;
 		TArray<FFastGeoSurrogateComponentDescriptor> LevelSurrogateComponentDescriptors;
-		for (TMap<AActor*, FTransformableActor>::TConstIterator It = TransformableActors.CreateConstIterator(); It; ++It)
+
+		for (const TPair<AActor*, FTransformableActor>& Entry : TransformableActors)
 		{
-			const TPair<AActor*, FTransformableActor>& Entry = *It;
 			AActor* Actor = Entry.Key;
 
 			FFastGeoComponentCluster* CurrentComponentCluster = LevelComponentCluster.Get();
@@ -408,8 +408,11 @@ void UFastGeoWorldPartitionRuntimeCellTransformer::Transform(ULevel* InLevel)
 						FastGeoPrimitiveComponent->ReserveSurrogateInstanceBodyIndices(NextAvailableInstanceBodyIndex);
 					}
 				}
+			}
 
-				// Remove the component from the actor
+			// Once all components are transformed, remove them from the actor
+			for (UActorComponent* Component : Components)
+			{
 				Actor->RemoveOwnedComponent(Component);
 				Component->MarkAsGarbage();
 			}
@@ -604,8 +607,9 @@ TMap<AActor*, TArray<AActor*>> UFastGeoWorldPartitionRuntimeCellTransformer::Bui
 	return ReferencedActors;
 }
 
-void UFastGeoWorldPartitionRuntimeCellTransformer::GatherTransformableActors(const TArray<AActor*>& InActors, const ULevel* InLevel, TMap<AActor*, FTransformableActor>& OutTransformableActors, FTransformationStats& OutStats)
+void UFastGeoWorldPartitionRuntimeCellTransformer::GatherTransformableActors(const TArray<AActor*>& InActors, const ULevel* InLevel, TArray<TPair<AActor*, FTransformableActor>>& OutTransformableActors, FTransformationStats& OutStats)
 {
+	TMap<AActor*, FTransformableActor> TransformableActors;
 	// Get transformation result for each actor
 	// This will retrieve the transformable components, and whether the actor is fully transformable (ie. if the actor can be deleted)
 	for (int32 ActorIndex = 0; ActorIndex < InActors.Num(); ++ActorIndex)
@@ -621,13 +625,18 @@ void UFastGeoWorldPartitionRuntimeCellTransformer::GatherTransformableActors(con
 			FFastGeoTransformResult ActorTransformResult = CanTransformActor(Actor, bIsActorFullyTransformable, TransformableComponents);
 			if (ActorTransformResult.GetResult() == EFastGeoTransform::Allow)
 			{
-				OutTransformableActors.FindOrAdd(Actor, { ActorIndex, bIsActorFullyTransformable, MoveTemp(TransformableComponents) });
+				// Sort components by pathname to ensure deterministic order
+				TransformableComponents.Sort([](const UActorComponent& A, const UActorComponent& B)
+				{
+					return A.GetPathName() < B.GetPathName();
+				});
+				TransformableActors.FindOrAdd(Actor, { ActorIndex, bIsActorFullyTransformable, MoveTemp(TransformableComponents) });
 			}
 		}
 	}
 
 	TMap<AActor*, TArray<AActor*>> ReferencedActors = BuildActorsReferencesMap(InLevel->Actors);
-	for (TMap<AActor*, FTransformableActor>::TIterator It = OutTransformableActors.CreateIterator(); It; ++It)
+	for (TMap<AActor*, FTransformableActor>::TIterator It = TransformableActors.CreateIterator(); It; ++It)
 	{
 		const TPair<AActor*, FTransformableActor>& Entry = *It;
 		const AActor* Actor = Entry.Key;
@@ -639,10 +648,10 @@ void UFastGeoWorldPartitionRuntimeCellTransformer::GatherTransformableActors(con
 		{
 			const AActor* Referencer = nullptr;
 			const FTransformableActor* TransformableReferencer = nullptr;
-			bool bHasNonFastGeoReferencer = Algo::AnyOf(*Referencers, [&Referencer, &TransformableReferencer, OutTransformableActors](const AActor* InReferencer)
+			bool bHasNonFastGeoReferencer = Algo::AnyOf(*Referencers, [&Referencer, &TransformableReferencer, &TransformableActors](const AActor* InReferencer)
 			{
 				Referencer = InReferencer;
-				TransformableReferencer = OutTransformableActors.Find(InReferencer);
+				TransformableReferencer = TransformableActors.Find(InReferencer);
 				return TransformableReferencer == nullptr || !TransformableReferencer->bIsActorFullyTransformable;
 			});
 
@@ -667,6 +676,17 @@ void UFastGeoWorldPartitionRuntimeCellTransformer::GatherTransformableActors(con
 		OutStats.PartiallyTransformableActorCount += !Entry.Value.bIsActorFullyTransformable ? 1 : 0;
 		OutStats.TransformedComponentCount += Entry.Value.TransformableComponents.Num();
 	}
+
+	// Sort by pathname to ensure deterministic order
+	OutTransformableActors.Reserve(TransformableActors.Num());
+	for (const TPair<AActor*, FTransformableActor>& Pair : TransformableActors)
+	{
+		OutTransformableActors.Add(Pair);
+	}
+	OutTransformableActors.Sort([](const TPair<AActor*, FTransformableActor>& A, const TPair<AActor*, FTransformableActor>& B)
+	{
+		return A.Key->GetPathName() < B.Key->GetPathName();
+	});
 }
 
 void UFastGeoWorldPartitionRuntimeCellTransformer::FTransformationStats::DumpStats(const TCHAR* InPrefixString)
diff --git a/Engine/Plugins/Experimental/FastGeoStreaming/Source/FastGeoStreaming/Public/FastGeoWorldPartitionRuntimeCellTransformer.h b/Engine/Plugins/Experimental/FastGeoStreaming/Source/FastGeoStreaming/Public/FastGeoWorldPartitionRuntimeCellTransformer.h
--- a/Engine/Plugins/Experimental/FastGeoStreaming/Source/FastGeoStreaming/Public/FastGeoWorldPartitionRuntimeCellTransformer.h
+++ b/Engine/Plugins/Experimental/FastGeoStreaming/Source/FastGeoStreaming/Public/FastGeoWorldPartitionRuntimeCellTransformer.h
@@ -224,7 +224,7 @@ protected:
 	};
 
 	TMap<AActor*, TArray<AActor*>> BuildActorsReferencesMap(const TArray<AActor*>& InActors);
-	void GatherTransformableActors(const TArray<AActor*>& InActors, const ULevel* InLevel, TMap<AActor*, FTransformableActor>& OutTransformableActors, FTransformationStats& OutStats);
+	void GatherTransformableActors(const TArray<AActor*>& InActors, const ULevel* InLevel, TArray<TPair<AActor*, FTransformableActor>>& OutTransformableActors, FTransformationStats& OutStats);
 #endif
 
 public:
