[FastGeo] Fixed typo in FastGeoSurrogateActor's root component name

#Synced-CL 50298075

diff --git a/Engine/Plugins/Experimental/FastGeoStreaming/Source/FastGeoStreaming/Private/FastGeoSurrogateActor.cpp b/Engine/Plugins/Experimental/FastGeoStreaming/Source/FastGeoStreaming/Private/FastGeoSurrogateActor.cpp
--- a/Engine/Plugins/Experimental/FastGeoStreaming/Source/FastGeoStreaming/Private/FastGeoSurrogateActor.cpp
+++ b/Engine/Plugins/Experimental/FastGeoStreaming/Source/FastGeoStreaming/Private/FastGeoSurrogateActor.cpp
@@ -12,7 +12,7 @@ AFastGeoSurrogateActor::AFastGeoSurrogateActor()
 {
 	PrimaryActorTick.bCanEverTick = false;
 	bReplicates = false;
-	RootComponent = CreateDefaultSubobject<USceneComponent>(TEXT("FastGeoCellProxyRootComponent"));
+	RootComponent = CreateDefaultSubobject<USceneComponent>(TEXT("FastGeoSurrogateRootComponent"));
 	RootComponent->SetMobility(EComponentMobility::Static);
 	SetRootComponent(RootComponent);
 }
