Fixed invalid check in UFastGeoContainer::Unregister() introduced by 50491879.

#Synced-CL 50298075

diff --git a/Engine/Plugins/Experimental/FastGeoStreaming/Source/FastGeoStreaming/Private/FastGeoContainer.cpp b/Engine/Plugins/Experimental/FastGeoStreaming/Source/FastGeoStreaming/Private/FastGeoContainer.cpp
--- a/Engine/Plugins/Experimental/FastGeoStreaming/Source/FastGeoStreaming/Private/FastGeoContainer.cpp
+++ b/Engine/Plugins/Experimental/FastGeoStreaming/Source/FastGeoStreaming/Private/FastGeoContainer.cpp
@@ -191,7 +191,7 @@ void UFastGeoContainer::Register()
 
 		if (const bool bHasAnyPendingTasks = HasAnyPendingTasks())
 		{
-			checkf(!bHasAnyPendingTasks, TEXT("UFastGeoContainer::Register() was called while container still has pending tasks. This will lead to invalid behavior."));
+			checkf(!bHasAnyPendingTasks, TEXT("UFastGeoContainer::Register() was called while the container still has pending tasks. This will lead to invalid behavior."));
 			return;
 		}
 
@@ -329,9 +329,9 @@ void UFastGeoContainer::Unregister()
 	{
 		check(!IsUsingSurrogateComponents() || SurrogateActor->HasRegisteredWithFastGeo());
 
-		if (const bool bHasAnyPendingTasks = HasAnyPendingTasks())
+		if (const bool bHasAnyPendingDestroyTasks = HasAnyPendingDestroyTasks())
 		{
-			checkf(!bHasAnyPendingTasks, TEXT("UFastGeoContainer::Unregister() was called while container still has pending tasks. This will lead to invalid behavior."));
+			checkf(!bHasAnyPendingDestroyTasks, TEXT("UFastGeoContainer::Unregister() was called while the container still has pending destroy tasks. This will lead to invalid behavior."));
 			return;
 		}
 
