Fix unreachable code errors and non-unity issues

#rnx

#Synced-CL 50791025

diff --git a/Engine/Plugins/Experimental/FastGeoStreaming/Source/FastGeoStreaming/Private/FastGeoPrimitiveComponent.cpp b/Engine/Plugins/Experimental/FastGeoStreaming/Source/FastGeoStreaming/Private/FastGeoPrimitiveComponent.cpp
--- a/Engine/Plugins/Experimental/FastGeoStreaming/Source/FastGeoStreaming/Private/FastGeoPrimitiveComponent.cpp
+++ b/Engine/Plugins/Experimental/FastGeoStreaming/Source/FastGeoStreaming/Private/FastGeoPrimitiveComponent.cpp
@@ -652,11 +652,11 @@ bool FFastGeoPrimitiveComponent::CheckPSOPrecachingAndBoostPriority(EPSOPrecache
 void FFastGeoPrimitiveComponent::SetupPrecachePSOParams(FPSOPrecacheParams& Params)
 {
 	const FPrimitiveSceneProxyDesc& SceneProxyDesc = GetSceneProxyDesc();
-	auto IsPrecomputedLightingValid = []() { return false; };
+	//auto IsPrecomputedLightingValid = []() { return false; };
 	Params.bRenderInMainPass = SceneProxyDesc.bRenderInMainPass;
 	Params.bRenderInDepthPass = SceneProxyDesc.bRenderInDepthPass;
 	Params.bStaticLighting = SceneProxyDesc.bHasStaticLighting;
-	Params.bUsesIndirectLightingCache = Params.bStaticLighting && SceneProxyDesc.IndirectLightingCacheQuality != ILCQ_Off && (!IsPrecomputedLightingValid() || SceneProxyDesc.LightmapType == ELightmapType::ForceVolumetric);
+	Params.bUsesIndirectLightingCache = Params.bStaticLighting && SceneProxyDesc.IndirectLightingCacheQuality != ILCQ_Off /*&& (!IsPrecomputedLightingValid() || SceneProxyDesc.LightmapType == ELightmapType::ForceVolumetric)*/;
 	Params.bAffectDynamicIndirectLighting = SceneProxyDesc.bAffectDynamicIndirectLighting;
 	Params.bCastShadow = SceneProxyDesc.CastShadow;
 	// Custom depth can be toggled at runtime with PSO precache call so assume it might be needed when depth pass is needed
diff --git a/Engine/Plugins/Runtime/GameplayAbilities/Source/GameplayAbilities/Private/Abilities/GameplayAbility.cpp b/Engine/Plugins/Runtime/GameplayAbilities/Source/GameplayAbilities/Private/Abilities/GameplayAbility.cpp
--- a/Engine/Plugins/Runtime/GameplayAbilities/Source/GameplayAbilities/Private/Abilities/GameplayAbility.cpp
+++ b/Engine/Plugins/Runtime/GameplayAbilities/Source/GameplayAbilities/Private/Abilities/GameplayAbility.cpp
@@ -1590,7 +1590,7 @@ void UGameplayAbility::OnGameplayTaskDeactivated(UGameplayTask& Task)
 
 	ActiveTasks.Remove(&Task);
 
-	if (ENABLE_ABILITYTASK_DEBUGMSG)
+	if constexpr (ENABLE_ABILITYTASK_DEBUGMSG)
 	{
 		AddAbilityTaskDebugMessage(&Task, TEXT("Ended."));
 	}
diff --git a/Engine/Plugins/Runtime/NetworkPrediction/Source/NetworkPrediction/Public/NetworkPredictionTrace.h b/Engine/Plugins/Runtime/NetworkPrediction/Source/NetworkPrediction/Public/NetworkPredictionTrace.h
--- a/Engine/Plugins/Runtime/NetworkPrediction/Source/NetworkPrediction/Public/NetworkPredictionTrace.h
+++ b/Engine/Plugins/Runtime/NetworkPrediction/Source/NetworkPrediction/Public/NetworkPredictionTrace.h
@@ -92,6 +92,8 @@ struct FNetworkPredictionInstanceConfig;
 
 #else
 
+#include "NetworkPredictionLog.h"
+
 // Compiled out
 #define UE_NP_TRACE_RECONCILE(Condition, Str) if (Condition) { return true; }
 #define UE_NP_TRACE_SIM(...)
diff --git a/Engine/Source/Runtime/Core/Public/Logging/LogMacros.h b/Engine/Source/Runtime/Core/Public/Logging/LogMacros.h
--- a/Engine/Source/Runtime/Core/Public/Logging/LogMacros.h
+++ b/Engine/Source/Runtime/Core/Public/Logging/LogMacros.h
@@ -229,8 +229,8 @@ CORE_API void BasicFatalLogWithProgramCounter(TStaticBasicLogRecordParam<FormatC
 		} \
 	}
 
-	#define UE_LOG_ACTIVE(...)				(0)
-	#define UE_LOG_ANY_ACTIVE(...)			(0)
+	#define UE_LOG_ACTIVE(...)				0
+	#define UE_LOG_ANY_ACTIVE(...)			0
 	#define UE_SUPPRESS(...) {}
 	#define UE_GET_LOG_VERBOSITY(...)		(::ELogVerbosity::NoLogging)
 	#define UE_SET_LOG_VERBOSITY(...)
