Changed IAsyncPhysicsStateProcessor::OnAsyncCreatePhysicsStateBegin_GameThread/OnAsyncDestroyPhysicsStateBegin_GameThread to output a unique set of objects to root.
Made FBodyInstancePhysicsCreationInputs::GetReferencedObjects returns valid unique set of objects.

#Synced-CL 50924777

diff --git a/Engine/Plugins/Experimental/FastGeoStreaming/Source/FastGeoStreaming/Internal/FastGeoContainer.h b/Engine/Plugins/Experimental/FastGeoStreaming/Source/FastGeoStreaming/Internal/FastGeoContainer.h
--- a/Engine/Plugins/Experimental/FastGeoStreaming/Source/FastGeoStreaming/Internal/FastGeoContainer.h
+++ b/Engine/Plugins/Experimental/FastGeoStreaming/Source/FastGeoStreaming/Internal/FastGeoContainer.h
@@ -170,10 +170,10 @@ public:
 	virtual bool AllowsAsyncPhysicsStateDestruction() const override final;
 	virtual bool IsAsyncPhysicsStateCreated() const override final;
 	virtual UObject* GetAsyncPhysicsStateObject() const override final;
-	virtual void OnAsyncCreatePhysicsStateBegin_GameThread(TArray<UObject*>& OutRootedObjects) override;
+	virtual void OnAsyncCreatePhysicsStateBegin_GameThread(TSet<UObject*>& OutRootedObjects) override;
 	virtual bool OnAsyncCreatePhysicsState(const UE::FTimeout& Timeout) override;
 	virtual void OnAsyncCreatePhysicsStateEnd_GameThread() override;
-	virtual void OnAsyncDestroyPhysicsStateBegin_GameThread(TArray<UObject*>& OutRootedObjects) override;
+	virtual void OnAsyncDestroyPhysicsStateBegin_GameThread(TSet<UObject*>& OutRootedObjects) override;
 	virtual bool OnAsyncDestroyPhysicsState(const UE::FTimeout& Timeout) override;
 	virtual void OnAsyncDestroyPhysicsStateEnd_GameThread() override;
 	virtual void CollectBodySetupsWithPhysicsMeshesToCreate(TSet<UBodySetup*>& OutBodySetups) const override;
diff --git a/Engine/Plugins/Experimental/FastGeoStreaming/Source/FastGeoStreaming/Private/FastGeoContainer.cpp b/Engine/Plugins/Experimental/FastGeoStreaming/Source/FastGeoStreaming/Private/FastGeoContainer.cpp
--- a/Engine/Plugins/Experimental/FastGeoStreaming/Source/FastGeoStreaming/Private/FastGeoContainer.cpp
+++ b/Engine/Plugins/Experimental/FastGeoStreaming/Source/FastGeoStreaming/Private/FastGeoContainer.cpp
@@ -930,7 +930,7 @@ UObject* UFastGeoContainer::GetAsyncPhysicsStateObject() const
 	return const_cast<UFastGeoContainer*>(this);
 }
 
-void UFastGeoContainer::OnAsyncCreatePhysicsStateBegin_GameThread(TArray<UObject*>& OutRootedObjects)
+void UFastGeoContainer::OnAsyncCreatePhysicsStateBegin_GameThread(TSet<UObject*>& OutRootedObjects)
 {
 	TRACE_CPUPROFILER_EVENT_SCOPE(UFastGeoContainer::OnAsyncCreatePhysicsStateBegin_GameThread);
 	check(PhysicsStateCreation == EPhysicsStateCreation::NotCreated);
@@ -982,7 +982,7 @@ void UFastGeoContainer::OnAsyncCreatePhysicsStateEnd_GameThread()
 	RegisterBodyInstances();
 }
 
-void UFastGeoContainer::OnAsyncDestroyPhysicsStateBegin_GameThread(TArray<UObject*>& OutRootedObjects)
+void UFastGeoContainer::OnAsyncDestroyPhysicsStateBegin_GameThread(TSet<UObject*>& OutRootedObjects)
 {
 	TRACE_CPUPROFILER_EVENT_SCOPE(UFastGeoContainer::OnAsyncDestroyPhysicsStateBegin_GameThread);
 	check(PhysicsStateCreation == EPhysicsStateCreation::Created);
diff --git a/Engine/Source/Runtime/Engine/Classes/Components/ActorComponent.h b/Engine/Source/Runtime/Engine/Classes/Components/ActorComponent.h
--- a/Engine/Source/Runtime/Engine/Classes/Components/ActorComponent.h
+++ b/Engine/Source/Runtime/Engine/Classes/Components/ActorComponent.h
@@ -899,7 +899,7 @@ public:
 	ENGINE_API virtual bool OnAsyncCreatePhysicsState(const UE::FTimeout& Timeout) override;
 	ENGINE_API virtual void OnAsyncCreatePhysicsStateEnd_GameThread() override;
 	virtual bool OnAsyncDestroyPhysicsState(const UE::FTimeout& Timeout) override { return true; }
-	ENGINE_API virtual void OnAsyncDestroyPhysicsStateBegin_GameThread(TArray<UObject*>& OutRootedObjects) override;
+	ENGINE_API virtual void OnAsyncDestroyPhysicsStateBegin_GameThread(TSet<UObject*>& OutRootedObjects) override;
 	ENGINE_API virtual void OnAsyncDestroyPhysicsStateEnd_GameThread() override;
 	//~ End IAsyncPhysicsStateProcessor
 	/**
diff --git a/Engine/Source/Runtime/Engine/Classes/Components/InstancedStaticMeshComponent.h b/Engine/Source/Runtime/Engine/Classes/Components/InstancedStaticMeshComponent.h
--- a/Engine/Source/Runtime/Engine/Classes/Components/InstancedStaticMeshComponent.h
+++ b/Engine/Source/Runtime/Engine/Classes/Components/InstancedStaticMeshComponent.h
@@ -528,7 +528,7 @@ protected:
 	ENGINE_API virtual void OnDestroyPhysicsState() override;
 	ENGINE_API virtual void OnAsyncCreatePhysicsStateEnd_GameThread() override;
 	ENGINE_API virtual bool OnAsyncDestroyPhysicsState(const UE::FTimeout& Timeout) override;
-	ENGINE_API virtual void OnAsyncDestroyPhysicsStateBegin_GameThread(TArray<UObject*>& OutRootedObjects) override;
+	ENGINE_API virtual void OnAsyncDestroyPhysicsStateBegin_GameThread(TSet<UObject*>& OutRootedObjects) override;
 
 public:
 	ENGINE_API virtual bool CanEditSimulatePhysics() override;
diff --git a/Engine/Source/Runtime/Engine/Classes/Components/PrimitiveComponent.h b/Engine/Source/Runtime/Engine/Classes/Components/PrimitiveComponent.h
--- a/Engine/Source/Runtime/Engine/Classes/Components/PrimitiveComponent.h
+++ b/Engine/Source/Runtime/Engine/Classes/Components/PrimitiveComponent.h
@@ -2522,10 +2522,10 @@ public:
 	//~ End UActorComponent Interface
 
 	//~ Begin IAsyncPhysicsStateProcessor
-	ENGINE_API virtual void OnAsyncCreatePhysicsStateBegin_GameThread(TArray<UObject*>& OutRootedObjects) override;
+	ENGINE_API virtual void OnAsyncCreatePhysicsStateBegin_GameThread(TSet<UObject*>& OutRootedObjects) override;
 	ENGINE_API virtual void OnAsyncCreatePhysicsStateEnd_GameThread() override;
 	ENGINE_API virtual bool OnAsyncDestroyPhysicsState(const UE::FTimeout& Timeout) override;
-	ENGINE_API virtual void OnAsyncDestroyPhysicsStateBegin_GameThread(TArray<UObject*>& OutRootedObjects) override;
+	ENGINE_API virtual void OnAsyncDestroyPhysicsStateBegin_GameThread(TSet<UObject*>& OutRootedObjects) override;
 	ENGINE_API virtual void OnAsyncDestroyPhysicsStateEnd_GameThread() override;
 	ENGINE_API virtual void CollectBodySetupsWithPhysicsMeshesToCreate(TSet<UBodySetup*>& OutBodySetups) const override;
 	//~ End IAsyncPhysicsStateProcessor
diff --git a/Engine/Source/Runtime/Engine/Classes/Components/SkeletalMeshComponent.h b/Engine/Source/Runtime/Engine/Classes/Components/SkeletalMeshComponent.h
--- a/Engine/Source/Runtime/Engine/Classes/Components/SkeletalMeshComponent.h
+++ b/Engine/Source/Runtime/Engine/Classes/Components/SkeletalMeshComponent.h
@@ -1866,10 +1866,10 @@ protected:
 public:
 	ENGINE_API virtual bool AllowsAsyncPhysicsStateCreation() const override;
 	ENGINE_API virtual bool AllowsAsyncPhysicsStateDestruction() const override;
-	ENGINE_API virtual void OnAsyncCreatePhysicsStateBegin_GameThread(TArray<UObject*>& OutRootedObjects) override;
+	ENGINE_API virtual void OnAsyncCreatePhysicsStateBegin_GameThread(TSet<UObject*>& OutRootedObjects) override;
 	ENGINE_API virtual bool OnAsyncCreatePhysicsState(const UE::FTimeout& TimeOut) override;
 	ENGINE_API virtual void OnAsyncCreatePhysicsStateEnd_GameThread() override;
-	ENGINE_API virtual void OnAsyncDestroyPhysicsStateBegin_GameThread(TArray<UObject*>& OutRootedObjects) override;
+	ENGINE_API virtual void OnAsyncDestroyPhysicsStateBegin_GameThread(TSet<UObject*>& OutRootedObjects) override;
 	ENGINE_API virtual bool OnAsyncDestroyPhysicsState(const UE::FTimeout& Timeout) override;
 
 	ENGINE_API virtual void InitializeComponent() override;
diff --git a/Engine/Source/Runtime/Engine/Classes/PhysicsEngine/BodyInstance.h b/Engine/Source/Runtime/Engine/Classes/PhysicsEngine/BodyInstance.h
--- a/Engine/Source/Runtime/Engine/Classes/PhysicsEngine/BodyInstance.h
+++ b/Engine/Source/Runtime/Engine/Classes/PhysicsEngine/BodyInstance.h
@@ -65,7 +65,7 @@ using FBodyInstancePhysicsCreationInputsPtr = TSharedPtr<FBodyInstancePhysicsCre
 // Cache used when creating physics body instances
 struct FBodyInstancePhysicsCreationInputs
 {
-	void GetReferencedObjects(TArray<UObject*>& OutReferencedObjects);
+	void GetReferencedObjects(TSet<UObject*>& OutReferencedObjects);
 
 	// Cached body setup resolved on GT
 	UBodySetup* ResolvedBodySetup = nullptr;
diff --git a/Engine/Source/Runtime/Engine/Private/Components/ActorComponent.cpp b/Engine/Source/Runtime/Engine/Private/Components/ActorComponent.cpp
--- a/Engine/Source/Runtime/Engine/Private/Components/ActorComponent.cpp
+++ b/Engine/Source/Runtime/Engine/Private/Components/ActorComponent.cpp
@@ -2378,7 +2378,7 @@ void UActorComponent::OnAsyncCreatePhysicsStateEnd_GameThread()
 	check(IsPreRegistered());
 }
 
-void UActorComponent::OnAsyncDestroyPhysicsStateBegin_GameThread(TArray<UObject*>& OutRootedObjects)
+void UActorComponent::OnAsyncDestroyPhysicsStateBegin_GameThread(TSet<UObject*>& OutRootedObjects)
 {
 	GlobalDestroyPhysicsDelegate.Broadcast(this);
 	ensureMsgf(bRegistered, TEXT("Component has physics state when not registered (%s)"), *GetFullName()); // should not have physics state unless we are registered
diff --git a/Engine/Source/Runtime/Engine/Private/Components/PrimitiveComponent.cpp b/Engine/Source/Runtime/Engine/Private/Components/PrimitiveComponent.cpp
--- a/Engine/Source/Runtime/Engine/Private/Components/PrimitiveComponent.cpp
+++ b/Engine/Source/Runtime/Engine/Private/Components/PrimitiveComponent.cpp
@@ -966,7 +966,7 @@ void UPrimitiveComponent::OnCreatePhysicsState()
 	}
 }
 
-void UPrimitiveComponent::OnAsyncCreatePhysicsStateBegin_GameThread(TArray<UObject*>& OutRootedObjects)
+void UPrimitiveComponent::OnAsyncCreatePhysicsStateBegin_GameThread(TSet<UObject*>& OutRootedObjects)
 {
 	BodyInstance.CachePhysicsCreationInputs(this)->GetReferencedObjects(OutRootedObjects);
 }
@@ -1073,7 +1073,7 @@ void UPrimitiveComponent::SendPhysicsTransform(ETeleportType Teleport)
 	BodyInstance.UpdateBodyScale(GetComponentTransform().GetScale3D());
 }
 
-void UPrimitiveComponent::OnAsyncDestroyPhysicsStateBegin_GameThread(TArray<UObject*>& OutRootedObjects)
+void UPrimitiveComponent::OnAsyncDestroyPhysicsStateBegin_GameThread(TSet<UObject*>& OutRootedObjects)
 {
 	check(AllowsAsyncPhysicsStateDestruction());
 
diff --git a/Engine/Source/Runtime/Engine/Private/InstancedStaticMesh.cpp b/Engine/Source/Runtime/Engine/Private/InstancedStaticMesh.cpp
--- a/Engine/Source/Runtime/Engine/Private/InstancedStaticMesh.cpp
+++ b/Engine/Source/Runtime/Engine/Private/InstancedStaticMesh.cpp
@@ -3013,7 +3013,7 @@ void UInstancedStaticMeshComponent::OnDestroyPhysicsState()
 	bNavigationRelevant = IsNavigationRelevant();
 }
 
-void UInstancedStaticMeshComponent::OnAsyncDestroyPhysicsStateBegin_GameThread(TArray<UObject*>& OutRootedObjects)
+void UInstancedStaticMeshComponent::OnAsyncDestroyPhysicsStateBegin_GameThread(TSet<UObject*>& OutRootedObjects)
 {
 	USceneComponent::OnAsyncDestroyPhysicsStateBegin_GameThread(OutRootedObjects);
 
diff --git a/Engine/Source/Runtime/Engine/Private/PhysicsEngine/BodyInstance.cpp b/Engine/Source/Runtime/Engine/Private/PhysicsEngine/BodyInstance.cpp
--- a/Engine/Source/Runtime/Engine/Private/PhysicsEngine/BodyInstance.cpp
+++ b/Engine/Source/Runtime/Engine/Private/PhysicsEngine/BodyInstance.cpp
@@ -5030,7 +5030,7 @@ const FBodyInstancePhysicsCreationInputsPtr& FBodyInstance::CachePhysicsCreation
 	return PhysicsCreationInputs;
 }
 
-void FBodyInstancePhysicsCreationInputs::GetReferencedObjects(TArray<UObject*>& OutReferencedObjects)
+void FBodyInstancePhysicsCreationInputs::GetReferencedObjects(TSet<UObject*>& OutReferencedObjects)
 {
 	if (ResolvedBodySetup)
 	{
@@ -5052,8 +5052,14 @@ void FBodyInstancePhysicsCreationInputs::GetReferencedObjects(TArray<UObject*>&
 
 	for (const FPhysicalMaterialMaskParams& Mask : ResolvedComplexMasks)
 	{
-		OutReferencedObjects.Add(Mask.PhysicalMaterialMap);
-		OutReferencedObjects.Add(Mask.PhysicalMaterialMask);
+		if (Mask.PhysicalMaterialMap)
+		{
+			OutReferencedObjects.Add(Mask.PhysicalMaterialMap);
+		}
+		if (Mask.PhysicalMaterialMask)
+		{
+			OutReferencedObjects.Add(Mask.PhysicalMaterialMask);
+		}
 	}
 }
 
diff --git a/Engine/Source/Runtime/Engine/Private/PhysicsEngine/Experimental/PhysScene_AsyncPhysicsStateJobQueue.cpp b/Engine/Source/Runtime/Engine/Private/PhysicsEngine/Experimental/PhysScene_AsyncPhysicsStateJobQueue.cpp
--- a/Engine/Source/Runtime/Engine/Private/PhysicsEngine/Experimental/PhysScene_AsyncPhysicsStateJobQueue.cpp
+++ b/Engine/Source/Runtime/Engine/Private/PhysicsEngine/Experimental/PhysScene_AsyncPhysicsStateJobQueue.cpp
@@ -389,7 +389,7 @@ void FPhysScene_AsyncPhysicsStateJobQueue::FJob::OnPreExecute_GameThread() const
 	// Root referenced objects
 	if (UObject* Obj = Processor->GetAsyncPhysicsStateObject())
 	{
-		Owner->AddRootObject(Obj);
+		RootedInputs.Add(Obj);
 	}
 	for (UObject* Obj : RootedInputs)
 	{
@@ -437,12 +437,6 @@ void FPhysScene_AsyncPhysicsStateJobQueue::FJob::OnPostExecute_GameThread() cons
 			}
 			break;
 		}
-
-		// Remove rooted objects
-		if (UObject* Obj = Processor->GetAsyncPhysicsStateObject())
-		{
-			Owner->RemoveRootObject(Obj);
-		}
 	}
 
 	for (UObject* Obj : RootedInputs)
diff --git a/Engine/Source/Runtime/Engine/Private/SkeletalMeshComponentPhysics.cpp b/Engine/Source/Runtime/Engine/Private/SkeletalMeshComponentPhysics.cpp
--- a/Engine/Source/Runtime/Engine/Private/SkeletalMeshComponentPhysics.cpp
+++ b/Engine/Source/Runtime/Engine/Private/SkeletalMeshComponentPhysics.cpp
@@ -1251,7 +1251,7 @@ bool USkeletalMeshComponent::AllowsAsyncPhysicsStateDestruction() const
 	return AllowsAsyncPhysicsStateCreation();
 }
 
-void USkeletalMeshComponent::OnAsyncCreatePhysicsStateBegin_GameThread(TArray<UObject*>& OutRootedObjects)
+void USkeletalMeshComponent::OnAsyncCreatePhysicsStateBegin_GameThread(TSet<UObject*>& OutRootedObjects)
 {
 	TRACE_CPUPROFILER_EVENT_SCOPE(USkeletalMeshComponent::OnAsyncCreatePhysicsStateBegin_GameThread);
 
@@ -1314,7 +1314,7 @@ void USkeletalMeshComponent::OnAsyncCreatePhysicsStateEnd_GameThread()
 	OnSkelMeshPhysicsCreated.Broadcast();
 }
 
-void USkeletalMeshComponent::OnAsyncDestroyPhysicsStateBegin_GameThread(TArray<UObject*>& OutRootedObjects)
+void USkeletalMeshComponent::OnAsyncDestroyPhysicsStateBegin_GameThread(TSet<UObject*>& OutRootedObjects)
 {
 	TRACE_CPUPROFILER_EVENT_SCOPE(USkeletalMeshComponent::OnAsyncDestroyPhysicsStateBegin_GameThread);
 
diff --git a/Engine/Source/Runtime/Engine/Public/Physics/Experimental/AsyncPhysicsStateProcessorInterface.h b/Engine/Source/Runtime/Engine/Public/Physics/Experimental/AsyncPhysicsStateProcessorInterface.h
--- a/Engine/Source/Runtime/Engine/Public/Physics/Experimental/AsyncPhysicsStateProcessorInterface.h
+++ b/Engine/Source/Runtime/Engine/Public/Physics/Experimental/AsyncPhysicsStateProcessorInterface.h
@@ -35,7 +35,7 @@ public:
 	 * 
 	 * @param OutRootedObjects Objects that must be rooted for the duration of the async creation.
 	 */
-	virtual void OnAsyncCreatePhysicsStateBegin_GameThread(TArray<UObject*>& OutRootedObjects) {}
+	virtual void OnAsyncCreatePhysicsStateBegin_GameThread(TSet<UObject*>& OutRootedObjects) {}
 
 	/** Used to create any physics state for this component outside of the GameThread. */
 	virtual bool OnAsyncCreatePhysicsState(const UE::FTimeout& TimeOut) { return true; }
@@ -53,7 +53,7 @@ public:
 	 * 
 	 * @param OutRootedObjects Objects that must be rooted for the duration of the async destruction.
 	 */
-	virtual void OnAsyncDestroyPhysicsStateBegin_GameThread(TArray<UObject*>& OutRootedObjects) {}
+	virtual void OnAsyncDestroyPhysicsStateBegin_GameThread(TSet<UObject*>& OutRootedObjects) {}
 
 	/** Used to destroy any physics state for this component outside of the GameThread. */
 	virtual bool OnAsyncDestroyPhysicsState(const UE::FTimeout& TimeOut) { return true; }
diff --git a/Engine/Source/Runtime/Engine/Public/Physics/Experimental/PhysScene_AsyncPhysicsStateJobQueue.h b/Engine/Source/Runtime/Engine/Public/Physics/Experimental/PhysScene_AsyncPhysicsStateJobQueue.h
--- a/Engine/Source/Runtime/Engine/Public/Physics/Experimental/PhysScene_AsyncPhysicsStateJobQueue.h
+++ b/Engine/Source/Runtime/Engine/Public/Physics/Experimental/PhysScene_AsyncPhysicsStateJobQueue.h
@@ -48,7 +48,7 @@ struct FPhysScene_AsyncPhysicsStateJobQueue : public FGCObject
 			return GetTypeHash(TPair<IAsyncPhysicsStateProcessor*, EJobType>{InJob.Processor, InJob.Type});
 		}
 
-		mutable TArray<UObject*> RootedInputs;
+		mutable TSet<UObject*> RootedInputs;
 		IAsyncPhysicsStateProcessor* Processor;
 		EJobType Type;
 	private:
