[FastGeo] Fixed fastgeo.show and actor coloration mode not working for transient FastGeoContainer (those created with UFastGeoContainer::CreateTransient, used by PCG componentless mode).

#Synced-CL 50924777

diff --git a/Engine/Plugins/Experimental/FastGeoStreaming/Source/FastGeoStreaming/Private/FastGeoWorldSubsystem.cpp b/Engine/Plugins/Experimental/FastGeoStreaming/Source/FastGeoStreaming/Private/FastGeoWorldSubsystem.cpp
--- a/Engine/Plugins/Experimental/FastGeoStreaming/Source/FastGeoStreaming/Private/FastGeoWorldSubsystem.cpp
+++ b/Engine/Plugins/Experimental/FastGeoStreaming/Source/FastGeoStreaming/Private/FastGeoWorldSubsystem.cpp
@@ -60,26 +60,21 @@ namespace FastGeo
 		FConsoleCommandWithArgsDelegate::CreateLambda([](const TArray<FString>& Args)
 		{
 			GShowFastGeo = (Args.Num() != 1) || (Args[0] != TEXT("0"));
-			for (const FWorldContext& Context : GEngine->GetWorldContexts())
+			for (TObjectIterator<UFastGeoContainer> It(RF_ClassDefaultObject, true, EInternalObjectFlags::Garbage); It; ++It)
 			{
-				UWorld* World = Context.World();
-				if (World && World->IsGameWorld())
+				UFastGeoContainer* FastGeo = *It;
+				if (FastGeo->GetWorld())
 				{
-					for (ULevel* Level : World->GetLevels())
+					FastGeo->ForEachComponentCluster([](FFastGeoComponentCluster& ComponentCluster)
 					{
-						if (UFastGeoContainer* FastGeo = Level->GetAssetUserData<UFastGeoContainer>())
-						{
-							FastGeo->ForEachComponentCluster([](FFastGeoComponentCluster& ComponentCluster)
-							{
-								ComponentCluster.UpdateVisibility();
-							});
-						}
-					}
+						ComponentCluster.UpdateVisibility();
+					});
 				}
 			}
 		})
 	);
 #endif
+
 }
 
 CSV_DEFINE_CATEGORY(FastGeo, true);
@@ -112,30 +107,22 @@ UFastGeoWorldSubsystem::UFastGeoWorldSubsystem()
 
 	auto UpdatePrimitivesColor = []()
 	{
-		if (GEngine)
+		for (TObjectIterator<UFastGeoContainer> It(RF_ClassDefaultObject, true, EInternalObjectFlags::Garbage); It; ++It)
 		{
-			for (const FWorldContext& Context : GEngine->GetWorldContexts())
+			UFastGeoContainer* FastGeo = *It;
+			UWorld* World = FastGeo->GetWorld();
+			if (World && World->IsGameWorld())
 			{
-				UWorld* World = Context.World();
-				if (World && World->IsGameWorld())
+				FastGeo->ForEachComponentCluster([](const FFastGeoComponentCluster& ComponentCluster)
 				{
-					for (ULevel* Level : World->GetLevels())
+					ComponentCluster.ForEachComponent<FFastGeoPrimitiveComponent>([](const FFastGeoPrimitiveComponent& Component)
 					{
-						if (UFastGeoContainer* FastGeo = Level->GetAssetUserData<UFastGeoContainer>())
+						if (FPrimitiveSceneProxy* SceneProxy = Component.GetSceneProxy())
 						{
-							FastGeo->ForEachComponentCluster([](const FFastGeoComponentCluster& ComponentCluster)
-							{
-								ComponentCluster.ForEachComponent<FFastGeoPrimitiveComponent>([](const FFastGeoPrimitiveComponent& Component)
-								{
-									if (FPrimitiveSceneProxy* SceneProxy = Component.GetSceneProxy())
-									{
-										SceneProxy->SetPrimitiveColor_GameThread(Component.GetDebugColor());
-									}
-								});
-							});
+							SceneProxy->SetPrimitiveColor_GameThread(Component.GetDebugColor());
 						}
-					}
-				}
+					});
+				});
 			}
 		}
 	};
