[FastGeo] Added option to disable asynchronous render state creation/destruction and PSO precaching (new CVar FastGeo.AllowAsyncRenderWork).
- Forced FastGeo.AllowAsyncRenderWork to false until issues with async version are fixed (UE-363084)
- Added UWorld::GetAddToWorldTimeout and UWorld::GetRemoveFromWorldTimeout for external systems like FastGeo to respect AddToWorld/RemoveFromWorld time budgets.
- FastGeoContainer now has internal registration states (current state, target state) to better deal with rapid target state changes while current state is still in progress.
- When FastGeo.AllowAsyncRenderWork is disabled:
. PSO Precaching is part of UWorld::AddToWorld and is done before creating render state (runs in a GT ParallelFor & time-sliced)
. Render proxy creation is part of UWorld::AddToWorld (runs in a GT ParallelFor & time-sliced)
. Render proxy destruction is part of UWorld::RemoveFromWorld (batched & time-sliced)

#Synced-CL 51093031

diff --git a/Engine/Plugins/Experimental/FastGeoStreaming/Source/FastGeoStreaming/Internal/FastGeoComponent.h b/Engine/Plugins/Experimental/FastGeoStreaming/Source/FastGeoStreaming/Internal/FastGeoComponent.h
--- a/Engine/Plugins/Experimental/FastGeoStreaming/Source/FastGeoStreaming/Internal/FastGeoComponent.h
+++ b/Engine/Plugins/Experimental/FastGeoStreaming/Source/FastGeoStreaming/Internal/FastGeoComponent.h
@@ -82,7 +82,6 @@ public:
 	UFastGeoContainer* GetOwnerContainer() const;
 	UWorld* GetWorld() const;
 	FSceneInterface* GetScene() const;
-	bool IsRegistered() const;
 	FLinearColor GetDebugColor() const;
 
 protected:
diff --git a/Engine/Plugins/Experimental/FastGeoStreaming/Source/FastGeoStreaming/Internal/FastGeoContainer.h b/Engine/Plugins/Experimental/FastGeoStreaming/Source/FastGeoStreaming/Internal/FastGeoContainer.h
--- a/Engine/Plugins/Experimental/FastGeoStreaming/Source/FastGeoStreaming/Internal/FastGeoContainer.h
+++ b/Engine/Plugins/Experimental/FastGeoStreaming/Source/FastGeoStreaming/Internal/FastGeoContainer.h
@@ -55,7 +55,13 @@ struct FFastGeoComponentClusterPendingState
 		int32 NumToProcess;
 		int32 NumProcessed;
 		int32 TotalNumProcessed;
-	} RenderState;
+	};
+
+	FRenderState RenderState;
+
+#if !WITH_EDITOR && UE_WITH_PSO_PRECACHING
+	FRenderState PrecachePSOState;
+#endif
 
 	struct FPhysicsState
 	{
@@ -89,13 +95,20 @@ struct FFastGeoComponentClusterPendingState
 	void Reset()
 	{
 		RenderState.Reset();
+#if !WITH_EDITOR && UE_WITH_PSO_PRECACHING
+		PrecachePSOState.Reset();
+#endif
 		PhysicsState.Reset();
 		bIsInBlockingWait = false;
 	}
 
 	bool HasAnyPendingState() const
 	{ 
-		return !RenderState.IsCompleted() || !PhysicsState.IsCompleted();
+		bool bHasAnyPendingState = !RenderState.IsCompleted() || !PhysicsState.IsCompleted();
+#if !WITH_EDITOR && UE_WITH_PSO_PRECACHING
+		bHasAnyPendingState |= !PrecachePSOState.IsCompleted();
+#endif
+		return bHasAnyPendingState;
 	}
 };
 
@@ -118,7 +131,12 @@ public:
 
 	void Register();
 	void Unregister();
+	bool IsRegistering() const;
 	bool IsRegistered() const;
+	bool IsUnregistering() const;
+	bool IsUnregistered() const;
+	bool IsFullyRegistered() const;
+	bool IsFullyUnregistered() const;
 	void Tick(bool bWaitForCompletion = false);
 	bool HasAnyPendingTasks() const;
 	bool HasAnyPendingCreateTasks() const;
@@ -213,6 +231,22 @@ protected:
 
 private:
 
+	void StartCreatePhysicsStateWork();
+	void StartDestroyPhysicsStateWork();
+	void TickPhysicsWork(bool bWaitForCompletion);
+
+	void StartPrecachePSOWork();
+	void StartCreateRenderStateWork();
+	void StartDestroyRenderStateWork();
+	void TickRenderWork(bool bWaitForCompletion);
+	void TickRenderWork_Sync(bool bWaitForCompletion);
+	bool TickPrecachePSO_Sync(const UE::FTimeout& Timeout);
+	bool TickCreateRenderState_Sync(const UE::FTimeout& Timeout);
+	bool TickDestroyRenderState_Sync(const UE::FTimeout& Timeout);
+	bool IsPrecachePSOCompleted() const;
+	bool IsCreateRenderStateCompleted() const;
+	bool IsDestroyRenderStateCompleted() const;
+
 	void CollectAssetReferences();
 	void SerializeComponentClusters(FArchive& Ar);
 	void RegisterToNavigationSystem();
@@ -220,6 +254,11 @@ private:
 	void OnNavigationInitDone(const UNavigationSystemBase& NavSys);
 	void RegisterBodyInstances();
 	void UnregisterBodyInstances();
+	void OnRegisterCompleted();
+	void OnUnregisterCompleted();
+	void ApplyRegistrationTargetState();
+	void StartRegisterTransition();
+	void StartUnregisterTransition();
 
 	template <typename ThisType, typename TComponentCluster, typename TFunc>
 	static void ForEachComponentCluster(ThisType* Self, TFunc&& InFunc)
@@ -287,14 +326,33 @@ private:
 			ForEachArray(Self->HLODs);
 	}
 
-	enum EPhysicsStateCreation
+	// UFastGeoContainer's physics state
+	enum class EPhysicsState : uint8
 	{
-		NotCreated,
+		None,
 		Creating,
 		Created,
 		Destroying
 	};
-	EPhysicsStateCreation PhysicsStateCreation = EPhysicsStateCreation::NotCreated;
+	EPhysicsState PhysicsState = EPhysicsState::None;
+
+	// UFastGeoContainer's registration state
+	enum class ERegistrationState : uint8
+	{
+		Unregistered,
+		Registering,
+		Registered,
+		Unregistering
+	};
+	ERegistrationState RegistrationState = ERegistrationState::Unregistered;
+
+	// UFastGeoContainer's registration target state
+	enum class ERegistrationTargetState : uint8
+	{
+		Registered,
+		Unregistered
+	};
+	ERegistrationTargetState RegistrationTargetState = ERegistrationTargetState::Unregistered;
 
 	// Persistent data
 	TArray<FFastGeoComponentCluster> ComponentClusters;
@@ -316,7 +374,6 @@ private:
 	TArray<FFastGeoComponent*> CollisionComponents;
 	TArray<FNavigationElementHandle> NavigationElementHandles;
 
-	bool bIsRegistered = false;
 	bool bIsBeingAsyncDestroyed = false;
 	int32 RegistrationEpoch = 0;
 
diff --git a/Engine/Plugins/Experimental/FastGeoStreaming/Source/FastGeoStreaming/Internal/FastGeoRegisteredComponent.h b/Engine/Plugins/Experimental/FastGeoStreaming/Source/FastGeoStreaming/Internal/FastGeoRegisteredComponent.h
--- a/Engine/Plugins/Experimental/FastGeoStreaming/Source/FastGeoStreaming/Internal/FastGeoRegisteredComponent.h
+++ b/Engine/Plugins/Experimental/FastGeoStreaming/Source/FastGeoStreaming/Internal/FastGeoRegisteredComponent.h
@@ -21,8 +21,11 @@ struct FFastGeoRegisteredComponent
 	{
 		if (T* ComponentPtr = Component.Get<T>())
 		{
-			if (ComponentPtr->GetOwnerContainer()->IsRegistered() &&
-				RegistrationEpoch == ComponentPtr->GetOwnerContainer()->RegistrationEpoch)
+			UFastGeoContainer* Container = ComponentPtr->GetOwnerContainer();
+			if (Container &&
+				Container->GetWorld() &&
+				(Container->IsRegistered() || Container->IsRegistering()) &&
+				(Container->RegistrationEpoch == RegistrationEpoch))
 			{
 				return ComponentPtr;
 			}
diff --git a/Engine/Plugins/Experimental/FastGeoStreaming/Source/FastGeoStreaming/Internal/FastGeoWorldSubsystem.h b/Engine/Plugins/Experimental/FastGeoStreaming/Source/FastGeoStreaming/Internal/FastGeoWorldSubsystem.h
--- a/Engine/Plugins/Experimental/FastGeoStreaming/Source/FastGeoStreaming/Internal/FastGeoWorldSubsystem.h
+++ b/Engine/Plugins/Experimental/FastGeoStreaming/Source/FastGeoStreaming/Internal/FastGeoWorldSubsystem.h
@@ -119,7 +119,8 @@ private:
 	void CommitAsyncRenderStateTasksBudget_Concurrent(float InUsedTimeBudgetMS, int32& InUsedComponentsBudget, int32 TimeEpoch);
 	void PushAsyncDestroyContainer(UFastGeoContainer* FastGeo);
 	FFastGeoAsyncRenderStateJobQueue* GetOrCreateJobQueue();
-	void WaitForAllAsyncJobsCompletion();
+	void WaitForAllPendingWorkCompletion();
+	void DrainPendingDestroyContainers();
 
 	UPROPERTY()
 	TSet<TObjectPtr<UFastGeoContainer>> PendingDestroyContainers;
diff --git a/Engine/Plugins/Experimental/FastGeoStreaming/Source/FastGeoStreaming/Private/FastGeoComponent.cpp b/Engine/Plugins/Experimental/FastGeoStreaming/Source/FastGeoStreaming/Private/FastGeoComponent.cpp
--- a/Engine/Plugins/Experimental/FastGeoStreaming/Source/FastGeoStreaming/Private/FastGeoComponent.cpp
+++ b/Engine/Plugins/Experimental/FastGeoStreaming/Source/FastGeoStreaming/Private/FastGeoComponent.cpp
@@ -86,11 +86,6 @@ FSceneInterface* FFastGeoComponent::GetScene() const
 	return nullptr;
 }
 
-bool FFastGeoComponent::IsRegistered() const
-{
-	return GetOwnerContainer()->IsRegistered();
-}
-
 FLinearColor FFastGeoComponent::GetDebugColor() const
 {
 	return UFastGeoWorldSubsystem::IsEnableDebugView() ? FLinearColor::Blue : FLinearColor::White;
diff --git a/Engine/Plugins/Experimental/FastGeoStreaming/Source/FastGeoStreaming/Private/FastGeoContainer.cpp b/Engine/Plugins/Experimental/FastGeoStreaming/Source/FastGeoStreaming/Private/FastGeoContainer.cpp
--- a/Engine/Plugins/Experimental/FastGeoStreaming/Source/FastGeoStreaming/Private/FastGeoContainer.cpp
+++ b/Engine/Plugins/Experimental/FastGeoStreaming/Source/FastGeoStreaming/Private/FastGeoContainer.cpp
@@ -31,8 +31,16 @@
 #include "WorldPartition/WorldPartition.h"
 #include "Engine/Engine.h"
 
-namespace FastGeo
+namespace FastGeo::Private
 {
+	// Async mode is disabled by default until issues are fixed
+	static bool GFastGeoAllowAsyncRenderWork = false;
+	static FAutoConsoleVariableRef CVarFastGeoAllowAsyncRenderWork(
+		TEXT("FastGeo.AllowAsyncRenderWork"),
+		GFastGeoAllowAsyncRenderWork,
+		TEXT("Allow FastGeo to perform render state creation/destruction and PSO precaching asynchronously."),
+		ECVF_ReadOnly);
+
 	static int32 GAsyncRenderStateTaskParallelWorkerCount = 1;
 	static FAutoConsoleVariableRef CVarAsyncRenderStateTaskParallelWorkerCount(
 		TEXT("FastGeo.AsyncRenderStateTask.ParallelWorkerCount"),
@@ -165,6 +173,35 @@ namespace FastGeo
 	};
 }
 
+void UFastGeoContainer::ApplyRegistrationTargetState()
+{
+	check(IsInGameThread());
+
+	switch (RegistrationState)
+	{
+	case ERegistrationState::Unregistered:
+		if (RegistrationTargetState == ERegistrationTargetState::Registered)
+		{
+			if (!bIsBeingAsyncDestroyed)
+			{
+				StartRegisterTransition();
+			}
+			UE_CLOG(bIsBeingAsyncDestroyed, LogFastGeoStreaming, Warning, TEXT("Can't start registering FastGeoContainer because it's being destroyed."));
+		}
+		break;
+
+	case ERegistrationState::Registered:
+		if (RegistrationTargetState == ERegistrationTargetState::Unregistered)
+		{
+			StartUnregisterTransition();
+		}
+		break;
+
+	default:
+		break;
+	}
+}
+
 void UFastGeoContainer::Register()
 {
 	TRACE_CPUPROFILER_EVENT_SCOPE(UFastGeoContainer::Register);
@@ -174,6 +211,7 @@ void UFastGeoContainer::Register()
 		return;
 	}
 #endif
+
 	if (!FFastGeoStreamingModule::IsFastGeoEnabled())
 	{
 		UE_LOG(LogFastGeoStreaming, Warning, TEXT("UFastGeoContainer::Register() ignored because FastGeo is disabled."));
@@ -182,81 +220,158 @@ void UFastGeoContainer::Register()
 
 	if (bIsBeingAsyncDestroyed)
 	{
-		UE_LOG(LogFastGeoStreaming, Warning, TEXT("UFastGeoContainer::Register() ignored because FastGeo is being destroyed."));
+		UE_LOG(LogFastGeoStreaming, Warning, TEXT("UFastGeoContainer::Register() ignored because FastGeoContainer is being destroyed."));
 		return;
 	}
 
-	if (!bIsRegistered)
+	RegistrationTargetState = ERegistrationTargetState::Registered;
+	ApplyRegistrationTargetState();
+
+	if (HasAnyPendingTasks())
 	{
-		check(!IsUsingSurrogateComponents() || !SurrogateActor->HasRegisteredWithFastGeo());
+		Tick();
+	}
+}
 
-		if (const bool bHasAnyPendingTasks = HasAnyPendingTasks())
-		{
-			checkf(!bHasAnyPendingTasks, TEXT("UFastGeoContainer::Register() was called while the container still has pending tasks. This will lead to invalid behavior."));
-			return;
-		}
+void UFastGeoContainer::StartRegisterTransition()
+{
+	TRACE_CPUPROFILER_EVENT_SCOPE(UFastGeoContainer::StartRegisterTransition);
+
+	check(IsInGameThread());
+#if WITH_EDITOR
+	check(!IsRunningCookCommandlet());
+#endif
+	check(FFastGeoStreamingModule::IsFastGeoEnabled());
+	check(!bIsBeingAsyncDestroyed);
+	check(RegistrationTargetState == ERegistrationTargetState::Registered);
+	check(RegistrationState == ERegistrationState::Unregistered);
+	check(!IsUsingSurrogateComponents() || !SurrogateActor->HasRegisteredWithFastGeo());
+
+	if (const bool bHasAnyPendingTasks = HasAnyPendingTasks())
+	{
+		checkf(!bHasAnyPendingTasks, TEXT("UFastGeoContainer::StartRegisterTransition() was called while the container still has pending tasks. This will lead to invalid behavior."));
+		return;
+	}
+
+	RegistrationState = ERegistrationState::Registering;
 
-		++RegistrationEpoch;
-		PendingCreate.Reset();
+	++RegistrationEpoch;
+	PendingCreate.Reset();
 
 #if WITH_EDITOR
-		// In PIE we need to initialize dynamic properties as there's no serialization
-		InitializeDynamicProperties(true);
+	// In PIE we need to initialize dynamic properties as there's no serialization
+	InitializeDynamicProperties(true);
 #endif
-		UWorldPartition* WorldPartition = FWorldPartitionHelpers::GetWorldPartition(this);
-		const bool bApplyWorldTransform = WorldPartition && WorldPartition->HasInstanceTransform();
-		const FTransform& Transform = WorldPartition ? WorldPartition->GetInstanceTransform() : FTransform::Identity;
-		ForEachComponentCluster([bApplyWorldTransform, &Transform, this](FFastGeoComponentCluster& ComponentCluster)
-		{
-			ComponentCluster.OnRegister();
+	UWorldPartition* WorldPartition = FWorldPartitionHelpers::GetWorldPartition(this);
+	const bool bApplyWorldTransform = WorldPartition && WorldPartition->HasInstanceTransform();
+	const FTransform& Transform = WorldPartition ? WorldPartition->GetInstanceTransform() : FTransform::Identity;
+
+	// Register clusters (FastGeoHLODs)
+	// Apply world transform
+	// Prepare PendingCreate.RenderState
+	ForEachComponentCluster([bApplyWorldTransform, &Transform, this](FFastGeoComponentCluster& ComponentCluster)
+	{
+		ComponentCluster.OnRegister();
 
-			ComponentCluster.ForEachComponent<FFastGeoComponent>([bApplyWorldTransform, &Transform, this](FFastGeoComponent& Component)
+		ComponentCluster.ForEachComponent<FFastGeoComponent>([bApplyWorldTransform, &Transform, this](FFastGeoComponent& Component)
+		{
+			if (bApplyWorldTransform)
 			{
-				if (bApplyWorldTransform)
-				{
-					Component.ApplyWorldTransform(Transform);
-				}
+				Component.ApplyWorldTransform(Transform);
+			}
 
-				if (Component.ShouldComponentAddToRenderScene())
-				{
-					PendingCreate.RenderState.ComponentsToProcess.Add(&Component);
-				}
-			});
+			if (Component.ShouldComponentAddToRenderScene())
+			{
+				PendingCreate.RenderState.ComponentsToProcess.Add(&Component);
+			}
 		});
+	});
 
-		UFastGeoWorldSubsystem* WorldSubsystem = GetWorld()->GetSubsystemChecked<UFastGeoWorldSubsystem>();
-		if (!PendingCreate.RenderState.ComponentsToProcess.IsEmpty())
-		{
 #if !WITH_EDITOR && UE_WITH_PSO_PRECACHING
-			check(bPrecachedPSOs);
-#endif
-			WorldSubsystem->PushAsyncCreateRenderStateJob(this);
+	// Prepare PendingCreate.PSOPrecachingState
+	PendingCreate.PrecachePSOState.Reset();
+	if (!FastGeo::Private::GFastGeoAllowAsyncRenderWork && !bPrecachedPSOs)
+	{
+		if (IsComponentPSOPrecachingEnabled())
+		{
+			PendingCreate.PrecachePSOState.ComponentsToProcess = PendingCreate.RenderState.ComponentsToProcess;
+			PendingCreate.PrecachePSOState.NumToProcess = PendingCreate.PrecachePSOState.ComponentsToProcess.Num();
 		}
+		bPrecachedPSOs = (PendingCreate.PrecachePSOState.ComponentsToProcess.Num() == 0);
+	}
+#endif
 
-		WorldSubsystem->PushAsyncCreatePhysicsStateJobs(this);
-
-		bIsRegistered = true;
+	StartCreateRenderStateWork();
+	StartCreatePhysicsStateWork();
 
-		const UNavigationSystemBase* NavigationSystem = GetWorld()->GetNavigationSystem();
-		if (NavigationSystem && NavigationSystem->IsWorldInitDone())
-		{
-			RegisterToNavigationSystem();
-		}
-		else
-		{
-			UNavigationSystemBase::OnNavigationInitDoneStaticDelegate().AddUObject(this, &UFastGeoContainer::OnNavigationInitDone);
-		}
+	// If no work needed for any system, registration is already complete
+	if (!HasAnyPendingCreateTasks())
+	{
+		OnRegisterCompleted();
 	}
 
-	if (PendingCreate.HasAnyPendingState())
+	const UNavigationSystemBase* NavigationSystem = GetWorld()->GetNavigationSystem();
+	if (NavigationSystem && NavigationSystem->IsWorldInitDone())
 	{
-		Tick();
+		RegisterToNavigationSystem();
 	}
+	else
+	{
+		UNavigationSystemBase::OnNavigationInitDoneStaticDelegate().AddUObject(this, &UFastGeoContainer::OnNavigationInitDone);
+	}
+}
+
+bool UFastGeoContainer::IsRegistering() const
+{
+	return RegistrationState == ERegistrationState::Registering;
 }
 
 bool UFastGeoContainer::IsRegistered() const
 {
-	return bIsRegistered;
+	return RegistrationState == ERegistrationState::Registered;
+}
+
+bool UFastGeoContainer::IsFullyRegistered() const
+{
+	return (RegistrationState == ERegistrationState::Registered) && !HasAnyPendingTasks();
+}
+
+bool UFastGeoContainer::IsUnregistering() const
+{
+	return RegistrationState == ERegistrationState::Unregistering;
+}
+
+bool UFastGeoContainer::IsUnregistered() const
+{
+	return RegistrationState == ERegistrationState::Unregistered;
+}
+
+bool UFastGeoContainer::IsFullyUnregistered() const
+{
+	return (RegistrationState == ERegistrationState::Unregistered) && !HasAnyPendingTasks();
+}
+
+void UFastGeoContainer::OnRegisterCompleted()
+{
+	TRACE_CPUPROFILER_EVENT_SCOPE(UFastGeoContainer::OnRegisterCompleted);
+	// Verify that registration is truly complete
+	check(!HasAnyPendingCreateTasks());
+	check(RegistrationState == ERegistrationState::Registering);
+	RegistrationState = ERegistrationState::Registered;
+	ApplyRegistrationTargetState();
+}
+
+void UFastGeoContainer::OnUnregisterCompleted()
+{
+	TRACE_CPUPROFILER_EVENT_SCOPE(UFastGeoContainer::OnUnregisterCompleted);
+	// Verify that unregistration is truly complete
+	check(!HasAnyPendingDestroyTasks());
+	check(RegistrationState == ERegistrationState::Unregistering);
+	RegistrationState = ERegistrationState::Unregistered;
+
+	OnUnregistered.Broadcast();
+
+	ApplyRegistrationTargetState();
 }
 
 UFastGeoContainer* UFastGeoContainer::CreateTransient(UWorld* InWorld)
@@ -326,61 +441,387 @@ void UFastGeoContainer::Unregister()
 	}
 #endif
 
-	if (bIsRegistered)
+	if (!FFastGeoStreamingModule::IsFastGeoEnabled())
+	{
+		UE_LOG(LogFastGeoStreaming, Warning, TEXT("UFastGeoContainer::Unregister() ignored because FastGeo is disabled."));
+		return;
+	}
+
+	RegistrationTargetState = ERegistrationTargetState::Unregistered;
+	ApplyRegistrationTargetState();
+
+	if (HasAnyPendingTasks())
+	{
+		Tick();
+	}
+}
+
+void UFastGeoContainer::StartUnregisterTransition()
+{
+	TRACE_CPUPROFILER_EVENT_SCOPE(UFastGeoContainer::StartUnregisterTransition);
+
+	check(IsInGameThread());
+#if WITH_EDITOR
+	check(!IsRunningCookCommandlet());
+#endif
+	check(FFastGeoStreamingModule::IsFastGeoEnabled());
+	check(RegistrationTargetState == ERegistrationTargetState::Unregistered);
+	check(RegistrationState == ERegistrationState::Registered);
+	check(!IsUsingSurrogateComponents() || SurrogateActor->HasRegisteredWithFastGeo());
+
+	if (const bool bHasAnyPendingDestroyTasks = HasAnyPendingDestroyTasks())
 	{
-		check(!IsUsingSurrogateComponents() || SurrogateActor->HasRegisteredWithFastGeo());
+		checkf(!bHasAnyPendingDestroyTasks, TEXT("UFastGeoContainer::StartUnregisterTransition() was called while the container still has pending destroy tasks. This will lead to invalid behavior."));
+		return;
+	}
 
-		if (const bool bHasAnyPendingDestroyTasks = HasAnyPendingDestroyTasks())
+	RegistrationState = ERegistrationState::Unregistering;
+
+	PendingDestroy.Reset();
+
+	ForEachComponentCluster([this](FFastGeoComponentCluster& ComponentCluster)
+	{
+		ComponentCluster.OnUnregister();
+
+		ComponentCluster.ForEachComponent<FFastGeoComponent>([this](FFastGeoComponent& Component)
 		{
-			checkf(!bHasAnyPendingDestroyTasks, TEXT("UFastGeoContainer::Unregister() was called while the container still has pending destroy tasks. This will lead to invalid behavior."));
-			return;
-		}
+			PendingDestroy.RenderState.ComponentsToProcess.Add(&Component);
+		});
+	});
 
-		PendingDestroy.Reset();
+	UnregisterFromNavigationSystem();
+		
+	StartDestroyRenderStateWork();
+	StartDestroyPhysicsStateWork();
 
-		ForEachComponentCluster([this](FFastGeoComponentCluster& ComponentCluster)
+	// If no work needed for any system, unregistration is already complete
+	if (!HasAnyPendingDestroyTasks())
+	{
+		OnUnregisterCompleted();
+	}
+}
+
+namespace FastGeo::Private
+{
+	template <typename FWorkFn>
+	FORCEINLINE bool AdvanceRenderStateBudgeted(FFastGeoComponentClusterPendingState::FRenderState& State, const UE::FTimeout& Timeout, FWorkFn&& WorkFn)
+	{
+		check(IsInGameThread());
+
+		if (State.IsCompleted())
+		{
+			return true;
+		}
+
+		if (Timeout.IsExpired())
 		{
-			ComponentCluster.OnUnregister();
+			return false;
+		}
+
+		std::atomic<int32> NextIndex{ 0 };
+		std::atomic<int32> NumProcessed{ 0 };
+		static constexpr int32 MinElementsPerTask = 8;
+		const int32 NumRemaining = State.NumToProcess;
+		const int32 StartIndex = State.TotalNumProcessed;
+		const int32 NumWorkers = FTaskGraphInterface::Get().GetNumWorkerThreads();
+		const int32 NumTasks = FMath::Clamp(FMath::DivideAndRoundUp(NumRemaining, MinElementsPerTask), 1, NumWorkers);
 
-			ComponentCluster.ForEachComponent<FFastGeoComponent>([this](FFastGeoComponent& Component)
+		ParallelFor(NumTasks, [&](int32 TaskIndex)
+		{
+			FTaskTagScope Scope(ETaskTag::EParallelGameThread);
+			while (true)
 			{
-				PendingDestroy.RenderState.ComponentsToProcess.Add(&Component);
-			});
+				if (Timeout.IsExpired())
+				{
+					return;
+				}
+
+				const int32 LocalIndex = NextIndex.fetch_add(1, std::memory_order_relaxed);
+				if (LocalIndex >= NumRemaining)
+				{
+					return;
+				}
+
+				FFastGeoComponent* Component = State.ComponentsToProcess[StartIndex + LocalIndex];
+				WorkFn(Component);
+
+				NumProcessed.fetch_add(1, std::memory_order_relaxed);
+			}
 		});
 
-		UnregisterFromNavigationSystem();
-		
-		UFastGeoWorldSubsystem* WorldSubsystem = GetWorld()->GetSubsystemChecked<UFastGeoWorldSubsystem>();
-		if (!PendingDestroy.RenderState.ComponentsToProcess.IsEmpty())
+		State.NumProcessed = NumProcessed.load(std::memory_order_relaxed);
+		State.TotalNumProcessed += State.NumProcessed;
+		State.NumToProcess = State.ComponentsToProcess.Num() - State.TotalNumProcessed;
+		return State.IsCompleted();
+	}
+}
+
+bool UFastGeoContainer::TickPrecachePSO_Sync(const UE::FTimeout& Timeout)
+{
+#if !WITH_EDITOR && UE_WITH_PSO_PRECACHING
+	if (!bPrecachedPSOs)
+	{
+		TRACE_CPUPROFILER_EVENT_SCOPE(UFastGeoContainer::TickPrecachePSO_Sync);
+		check(!FastGeo::Private::GFastGeoAllowAsyncRenderWork);
+
+		bPrecachedPSOs = FastGeo::Private::AdvanceRenderStateBudgeted(PendingCreate.PrecachePSOState, Timeout, [](FFastGeoComponent* Component)
+		{
+			TRACE_CPUPROFILER_EVENT_SCOPE(UFastGeoContainer::PrecachePSOs);
+			Component->PrecachePSOs_Concurrent();
+		});
+
+		if (bPrecachedPSOs)
+		{
+			PendingCreate.PrecachePSOState.Reset();
+		}
+	}
+	return bPrecachedPSOs;
+#endif
+	return true;
+}
+
+bool UFastGeoContainer::TickCreateRenderState_Sync(const UE::FTimeout& Timeout)
+{
+	TRACE_CPUPROFILER_EVENT_SCOPE(UFastGeoContainer::TickCreateRenderState_Sync);
+	check(IsInGameThread());
+	check(!FastGeo::Private::GFastGeoAllowAsyncRenderWork);
+
+	if (PendingCreate.RenderState.ComponentsToProcess.Num() > 0 && PendingCreate.RenderState.NumToProcess == 0)
+	{
+		OnCreateRenderStateBegin_GameThread();
+	}
+
+	const bool bCompleted = FastGeo::Private::AdvanceRenderStateBudgeted(PendingCreate.RenderState, Timeout, [](FFastGeoComponent* Component)
+	{
+		TRACE_CPUPROFILER_EVENT_SCOPE(UFastGeoContainer::CreateRenderState);
+		Component->CreateRenderState(nullptr);
+	});
+
+	if (bCompleted)
+	{
+		// Force NumProcessed to zero because OnCreateRenderStateEnd_GameThread was designed to increment TotalNumProcessed
+		PendingCreate.RenderState.NumProcessed = 0;
+		OnCreateRenderStateEnd_GameThread();
+	}
+
+	return bCompleted;
+}
+
+bool UFastGeoContainer::TickDestroyRenderState_Sync(const UE::FTimeout& Timeout)
+{
+	if (Timeout.IsExpired())
+	{
+		return false;
+	}
+
+	TRACE_CPUPROFILER_EVENT_SCOPE(UFastGeoContainer::TickDestroyRenderState_Sync);
+	check(IsInGameThread());
+	check(!FastGeo::Private::GFastGeoAllowAsyncRenderWork);
+
+	FFastGeoComponentClusterPendingState::FRenderState& State = PendingDestroy.RenderState;
+	if (State.ComponentsToProcess.Num() > 0 && State.NumToProcess == 0)
+	{
+		OnDestroyRenderStateBegin_GameThread();
+	}
+	
+	static constexpr int32 MaxComponentPerIteration = 16;
+	const TArray<FFastGeoComponent*>& Components = State.ComponentsToProcess;
+	const int32 ComponentsCount = Components.Num();
+	while (!IsDestroyRenderStateCompleted())
+	{
+		if (Timeout.IsExpired())
 		{
-			WorldSubsystem->PushAsyncDestroyRenderStateJob(this);
+			return false;
 		}
 
+		{
+			FFastGeoDestroyRenderStateContext Context(GetWorld()->Scene);
+			const int32 StartIndex = State.TotalNumProcessed;
+			const int32 EndIndex = FMath::Min<int32>(StartIndex + MaxComponentPerIteration, ComponentsCount);
+			for (int Index = StartIndex; Index < EndIndex; ++Index)
+			{
+				Components[Index]->DestroyRenderState(&Context);
+			}
+			State.TotalNumProcessed += (EndIndex - StartIndex);
+			State.NumProcessed = State.TotalNumProcessed;
+			State.NumToProcess = State.ComponentsToProcess.Num() - State.TotalNumProcessed;
+		}
+	}
+
+	check(IsDestroyRenderStateCompleted());
+	// Force NumProcessed to zero because OnDestroyRenderStateEnd_GameThread was designed to increment TotalNumProcessed
+	State.NumProcessed = 0;
+	OnDestroyRenderStateEnd_GameThread();
+	return true;
+}
+
+void UFastGeoContainer::StartCreatePhysicsStateWork()
+{
+	if (const bool bHasPhysicsWork = !CollisionComponents.IsEmpty())
+	{
+		UFastGeoWorldSubsystem* WorldSubsystem = GetWorld()->GetSubsystemChecked<UFastGeoWorldSubsystem>();
+		WorldSubsystem->PushAsyncCreatePhysicsStateJobs(this);
+	}
+}
+
+void UFastGeoContainer::StartDestroyPhysicsStateWork()
+{
+	if (const bool bHasPhysicsWork = !CollisionComponents.IsEmpty())
+	{
+		UFastGeoWorldSubsystem* WorldSubsystem = GetWorld()->GetSubsystemChecked<UFastGeoWorldSubsystem>();
 		WorldSubsystem->PushAsyncDestroyPhysicsStateJobs(this);
-		bIsRegistered = false;
+	}
+}
 
-		OnUnregistered.Broadcast();
+void UFastGeoContainer::TickPhysicsWork(bool bWaitForCompletion)
+{
+	if (FPhysScene* PhysScene = GetWorld()->GetPhysicsScene())
+	{
+		PhysScene->ProcessAsyncPhysicsStateJobs(bWaitForCompletion);
 	}
+}
 
-	if (PendingDestroy.HasAnyPendingState())
+void UFastGeoContainer::StartPrecachePSOWork()
+{
+#if !WITH_EDITOR && UE_WITH_PSO_PRECACHING
+	if (FastGeo::Private::GFastGeoAllowAsyncRenderWork && !bPrecachedPSOs)
 	{
-		Tick();
+		bPrecachedPSOs = true;
+
+		if (IsComponentPSOPrecachingEnabled())
+		{
+			UFastGeoWorldSubsystem* WorldSubsystem = GetWorld()->GetSubsystemChecked<UFastGeoWorldSubsystem>();
+			WorldSubsystem->PushAsyncPrecachePSOsJob(this);
+		}
+	}
+#endif
+}
+
+void UFastGeoContainer::StartCreateRenderStateWork()
+{
+	if (FastGeo::Private::GFastGeoAllowAsyncRenderWork && !IsCreateRenderStateCompleted())
+	{
+#if !WITH_EDITOR && UE_WITH_PSO_PRECACHING
+		// PSO precaching should have been started by UFastGeoWorldSubsystem::OnLevelStreamingStateChanged
+		if (!ensure(bPrecachedPSOs))
+		{
+			StartPrecachePSOWork();
+			check(bPrecachedPSOs);
+		}
+#endif
+		UFastGeoWorldSubsystem* WorldSubsystem = GetWorld()->GetSubsystemChecked<UFastGeoWorldSubsystem>();
+		WorldSubsystem->PushAsyncCreateRenderStateJob(this);
+	}
+}
+
+void UFastGeoContainer::StartDestroyRenderStateWork()
+{
+	if (FastGeo::Private::GFastGeoAllowAsyncRenderWork && !IsDestroyRenderStateCompleted())
+	{
+		UFastGeoWorldSubsystem* WorldSubsystem = GetWorld()->GetSubsystemChecked<UFastGeoWorldSubsystem>();
+		WorldSubsystem->PushAsyncDestroyRenderStateJob(this);
+	}
+}
+
+bool UFastGeoContainer::IsPrecachePSOCompleted() const
+{
+#if !WITH_EDITOR && UE_WITH_PSO_PRECACHING
+	return bPrecachedPSOs && PendingCreate.PrecachePSOState.IsCompleted();
+#else
+	return true;
+#endif
+}
+
+bool UFastGeoContainer::IsCreateRenderStateCompleted() const
+{
+	return PendingCreate.RenderState.IsCompleted();
+}
+
+bool UFastGeoContainer::IsDestroyRenderStateCompleted() const
+{
+	return PendingDestroy.RenderState.IsCompleted();
+}
+
+void UFastGeoContainer::TickRenderWork(bool bWaitForCompletion)
+{
+	if (FastGeo::Private::GFastGeoAllowAsyncRenderWork)
+	{
+		UFastGeoWorldSubsystem* WorldSubsystem = GetWorld()->GetSubsystemChecked<UFastGeoWorldSubsystem>();
+		WorldSubsystem->ProcessAsyncRenderStateJobs(bWaitForCompletion);
+	}
+	else
+	{
+		TickRenderWork_Sync(bWaitForCompletion);
+	}
+}
+
+void UFastGeoContainer::TickRenderWork_Sync(bool bWaitForCompletion)
+{
+	if (FastGeo::Private::GFastGeoAllowAsyncRenderWork)
+	{
+		return;
+	}
+
+	TRACE_CPUPROFILER_EVENT_SCOPE(UFastGeoContainer::TickRenderWork_Sync);
+	check(IsInGameThread());
+
+	switch (RegistrationState)
+	{
+	case ERegistrationState::Registering:
+		{
+			if (!IsPrecachePSOCompleted() || !IsCreateRenderStateCompleted())
+			{
+				const UE::FTimeout Timeout = bWaitForCompletion ? UE::FTimeout::Never() : GetWorld()->GetAddToWorldTimeout();
+				if (TickPrecachePSO_Sync(Timeout))
+				{
+					check(IsPrecachePSOCompleted());
+					if (TickCreateRenderState_Sync(Timeout))
+					{
+						check(IsCreateRenderStateCompleted());
+					}
+				}
+			}
+		}
+		break;
+	case ERegistrationState::Unregistering:
+		{
+			if (!IsDestroyRenderStateCompleted())
+			{
+				const UE::FTimeout Timeout = bWaitForCompletion ? UE::FTimeout::Never() : GetWorld()->GetRemoveFromWorldTimeout();
+				if (TickDestroyRenderState_Sync(Timeout))
+				{
+					check(IsDestroyRenderStateCompleted());
+				}
+			}
+		}
+		break;
+	default:
+		check(!HasAnyPendingTasks());
+		break;
 	}
 }
 
 void UFastGeoContainer::Tick(bool bWaitForCompletion)
 {
-	UWorld* World = GetWorld();
-	check(World);
-	UFastGeoWorldSubsystem* WorldSubsystem = World->GetSubsystemChecked<UFastGeoWorldSubsystem>();
+	check(IsInGameThread());
+	if (!ensureMsgf(GetWorld(), TEXT("UFastGeoContainer::Tick() skipped, world is not valid.")))
+	{
+		return;
+	}
+
+	if (!HasAnyPendingTasks())
+	{
+		ApplyRegistrationTargetState();
+	}
 
 	do
 	{
-		WorldSubsystem->ProcessAsyncRenderStateJobs(bWaitForCompletion);
+		TickRenderWork(bWaitForCompletion);
+		TickPhysicsWork(bWaitForCompletion);
 
-		if (FPhysScene* PhysScene = World->GetPhysicsScene())
+		if (bWaitForCompletion && !HasAnyPendingTasks())
 		{
-			PhysScene->ProcessAsyncPhysicsStateJobs(bWaitForCompletion);
+			ApplyRegistrationTargetState();
 		}
 	} while (bWaitForCompletion && HasAnyPendingTasks());
 }
@@ -393,13 +834,13 @@ bool UFastGeoContainer::HasAnyPendingTasks() const
 bool UFastGeoContainer::HasAnyPendingCreateTasks() const
 {
 	// Physics state can have no PendingState but OnAsyncCreatePhysicsStateEnd_GameThread has not been called yet
-	return (PhysicsStateCreation == EPhysicsStateCreation::Creating) || PendingCreate.HasAnyPendingState();
+	return (PhysicsState == EPhysicsState::Creating) || PendingCreate.HasAnyPendingState();
 }
 
 bool UFastGeoContainer::HasAnyPendingDestroyTasks() const
 {
 	// Physics state can have no PendingState but OnAsyncDestroyPhysicsStateEnd_GameThread has not been called yet
-	return (PhysicsStateCreation == EPhysicsStateCreation::Destroying) || PendingDestroy.HasAnyPendingState();
+	return (PhysicsState == EPhysicsState::Destroying) || PendingDestroy.HasAnyPendingState();
 }
 
 FFastGeoComponentCluster* UFastGeoContainer::GetComponentCluster(uint32 InComponentClusterTypeID, int32 InComponentClusterIndex)
@@ -433,8 +874,8 @@ ULevel* UFastGeoContainer::GetLevel() const
 
 UWorld* UFastGeoContainer::GetWorld() const
 {
-	check(GetLevel());
-	return GetLevel()->GetWorld();
+	ULevel* Level = GetLevel();
+	return Level ? Level->GetWorld() : nullptr;
 }
 
 void UFastGeoContainer::OnNavigationInitDone(const UNavigationSystemBase& InNavigationSystem)
@@ -456,7 +897,7 @@ void UFastGeoContainer::RegisterToNavigationSystem()
 	const UNavigationSystemBase* NavigationSystem = World->GetNavigationSystem();
 	check(NavigationSystem);
 	check(NavigationSystem->IsWorldInitDone());
-	check(bIsRegistered);
+	check(IsRegistered() || IsRegistering());
 
 	if (!FNavigationSystem::SupportsDynamicChanges(World))
 	{
@@ -536,9 +977,11 @@ void UFastGeoContainer::OnDestroyRenderStateBegin_GameThread()
 
 void UFastGeoContainer::OnCreateRenderState_Concurrent()
 {
-	static const int32 MinNumElementsToProcessPerThread = 8;
+	check(FastGeo::Private::GFastGeoAllowAsyncRenderWork);
+
+	static constexpr int32 MinNumElementsToProcessPerThread = 8;
 	const int32 NumComponentsToProcess = PendingCreate.RenderState.NumToProcess;
-	const int32 MaxNumThreads = PendingCreate.bIsInBlockingWait ? INT32_MAX : FastGeo::GAsyncRenderStateTaskParallelWorkerCount;
+	const int32 MaxNumThreads = PendingCreate.bIsInBlockingWait ? INT32_MAX : FastGeo::Private::GAsyncRenderStateTaskParallelWorkerCount;
 	const int32 NumThreads = FMath::Clamp(NumComponentsToProcess / MinNumElementsToProcessPerThread, 1, MaxNumThreads);
 	const bool bIsParallelForAllowed = NumThreads > 1 && FApp::ShouldUseThreadingForPerformance();
 
@@ -553,8 +996,8 @@ void UFastGeoContainer::OnCreateRenderState_Concurrent()
 	const double TimeBudgetSeconds = AvailableTimeBudgetMS / 1000.0;
 	UE::FTimeout Timeout = UE::FTimeout(TimeBudgetSeconds);
 
-	TAtomic<int32> NextIndex = 0;
-	TAtomic<int32> NumProcessed = 0;
+	std::atomic<int32> NextIndex{ 0 };
+	std::atomic<int32> NumProcessed{ 0 };
 
 	if (ComponentsBudget > 0 && !Timeout.IsExpired())
 	{
@@ -573,22 +1016,22 @@ void UFastGeoContainer::OnCreateRenderState_Concurrent()
 				}
 
 				// All work completed ?
-				int32 LocalIdx = NextIndex++;
-				if (LocalIdx >= ComponentsBudget)
+				const int32 LocalIndex = NextIndex.fetch_add(1, std::memory_order_relaxed);
+				if (LocalIndex >= ComponentsBudget)
 				{
 					return; 
 				}
 
-				FFastGeoComponent* ComponentToProcess = PendingCreate.RenderState.ComponentsToProcess[PendingCreate.RenderState.TotalNumProcessed + LocalIdx];
+				FFastGeoComponent* ComponentToProcess = PendingCreate.RenderState.ComponentsToProcess[PendingCreate.RenderState.TotalNumProcessed + LocalIndex];
 				ComponentToProcess->CreateRenderState(/*Context = */ nullptr);
 
-				NumProcessed++;
+				NumProcessed.fetch_add(1, std::memory_order_relaxed);
 			}
 		}, bIsParallelForAllowed ? EParallelForFlags::None : EParallelForFlags::ForceSingleThread);
 	}
 
 	// Adjust NumRenderStatesToProcess to the actual number of states we processed. 
-	PendingCreate.RenderState.NumProcessed = NumProcessed.Load(EMemoryOrder::Relaxed);
+	PendingCreate.RenderState.NumProcessed = NumProcessed.load(std::memory_order_relaxed);
 
 	WorldSubsystem->CommitAsyncRenderStateTasksBudget_Concurrent(Timeout.GetElapsedSeconds() * 1000, PendingCreate.RenderState.NumProcessed, TimeEpoch);
 }
@@ -613,35 +1056,57 @@ void UFastGeoContainer::OnDestroyRenderState_Concurrent()
 
 void UFastGeoContainer::OnCreateRenderStateEnd_GameThread()
 {
-	check(PendingCreate.RenderState.NumToProcess != 0);
-	check(PendingCreate.RenderState.TotalNumProcessed != PendingCreate.RenderState.ComponentsToProcess.Num());
+	if (FastGeo::Private::GFastGeoAllowAsyncRenderWork)
+	{
+		check(PendingCreate.RenderState.NumToProcess != 0);
+		check(PendingCreate.RenderState.TotalNumProcessed != PendingCreate.RenderState.ComponentsToProcess.Num());
+	}
 
 	PendingCreate.RenderState.TotalNumProcessed += PendingCreate.RenderState.NumProcessed;
 	PendingCreate.RenderState.NumToProcess = 0;
 	PendingCreate.RenderState.NumProcessed = 0;
 	PendingCreate.bIsInBlockingWait = false;
 
-	if (PendingCreate.RenderState.TotalNumProcessed < PendingCreate.RenderState.ComponentsToProcess.Num())
+	if (!IsCreateRenderStateCompleted())
 	{
-		UFastGeoWorldSubsystem* WorldSubsystem = GetWorld()->GetSubsystemChecked<UFastGeoWorldSubsystem>();
-		WorldSubsystem->PushAsyncCreateRenderStateJob(this);
+		check(FastGeo::Private::GFastGeoAllowAsyncRenderWork);
+		StartCreateRenderStateWork();
+	}
+	else
+	{
+		// Render work complete - check if all registration systems are done
+		if (!HasAnyPendingCreateTasks())
+	{
+			OnRegisterCompleted();
+		}
 	}
 }
 
 void UFastGeoContainer::OnDestroyRenderStateEnd_GameThread()
 {
-	check(PendingDestroy.RenderState.NumToProcess != 0);
-	check(PendingDestroy.RenderState.TotalNumProcessed != PendingDestroy.RenderState.ComponentsToProcess.Num());
+	if (FastGeo::Private::GFastGeoAllowAsyncRenderWork)
+	{
+		check(PendingDestroy.RenderState.NumToProcess != 0);
+		check(PendingDestroy.RenderState.TotalNumProcessed != PendingDestroy.RenderState.ComponentsToProcess.Num());
+	}
 
 	PendingDestroy.RenderState.TotalNumProcessed += PendingDestroy.RenderState.NumProcessed;
 	PendingDestroy.RenderState.NumToProcess = 0;
 	PendingDestroy.RenderState.NumProcessed = 0;	
 	PendingDestroy.bIsInBlockingWait = false;
 
-	if (PendingDestroy.RenderState.TotalNumProcessed < PendingDestroy.RenderState.ComponentsToProcess.Num())
+	if (!IsDestroyRenderStateCompleted())
 	{
-		UFastGeoWorldSubsystem* WorldSubsystem = GetWorld()->GetSubsystemChecked<UFastGeoWorldSubsystem>();
-		WorldSubsystem->PushAsyncDestroyRenderStateJob(this);
+		check(FastGeo::Private::GFastGeoAllowAsyncRenderWork);
+		StartDestroyRenderStateWork();
+	}
+	else
+	{
+		// Render work complete - check if all unregistration systems are done
+		if (!HasAnyPendingDestroyTasks())
+		{
+			OnUnregisterCompleted();
+		}
 	}
 }
 
@@ -782,7 +1247,7 @@ void UFastGeoContainer::Serialize(FArchive& Ar)
 {
 	Super::Serialize(Ar);
 
-	FastGeo::FAssetRemapArchive AssetRemapAr(Ar, Assets);
+	FastGeo::Private::FAssetRemapArchive AssetRemapAr(Ar, Assets);
 	SerializeComponentClusters(AssetRemapAr);
 
 #if !WITH_EDITOR
@@ -819,20 +1284,9 @@ void UFastGeoContainer::OnPrecachePSOs_Concurrent()
 
 void UFastGeoContainer::PrecachePSOs()
 {
+	// This function is currently only starting precaching when running in async mode
 	TRACE_CPUPROFILER_EVENT_SCOPE(UFastGeoContainer::PrecachePSOs_GameThread);
-
-#if !WITH_EDITOR && UE_WITH_PSO_PRECACHING
-	if (!bPrecachedPSOs)
-	{
-		bPrecachedPSOs = true;
-
-		if (IsComponentPSOPrecachingEnabled())
-		{
-			UFastGeoWorldSubsystem* WorldSubsystem = GetWorld()->GetSubsystemChecked<UFastGeoWorldSubsystem>();
-			WorldSubsystem->PushAsyncPrecachePSOsJob(this);
-		}
-	}
-#endif
+	StartPrecachePSOWork();
 }
 
 void UFastGeoContainer::InitializeDynamicProperties(bool bInitForPlay)
@@ -922,7 +1376,7 @@ bool UFastGeoContainer::AllowsAsyncPhysicsStateDestruction() const
 
 bool UFastGeoContainer::IsAsyncPhysicsStateCreated() const
 {
-	return PhysicsStateCreation == EPhysicsStateCreation::Created;
+	return PhysicsState == EPhysicsState::Created;
 }
 
 UObject* UFastGeoContainer::GetAsyncPhysicsStateObject() const
@@ -933,10 +1387,10 @@ UObject* UFastGeoContainer::GetAsyncPhysicsStateObject() const
 void UFastGeoContainer::OnAsyncCreatePhysicsStateBegin_GameThread(TSet<UObject*>& OutRootedObjects)
 {
 	TRACE_CPUPROFILER_EVENT_SCOPE(UFastGeoContainer::OnAsyncCreatePhysicsStateBegin_GameThread);
-	check(PhysicsStateCreation == EPhysicsStateCreation::NotCreated);
+	check(PhysicsState == EPhysicsState::None);
 	check(PendingCreate.PhysicsState.ComponentsToProcess);
 	check(PendingCreate.PhysicsState.ComponentsToProcess->Num());
-	PhysicsStateCreation = EPhysicsStateCreation::Creating;
+	PhysicsState = EPhysicsState::Creating;
 
 	// No need to fill OutRootedObjects as Assets already keeps objects alive
 	bool bSingleThreaded = !FApp::ShouldUseThreadingForPerformance();
@@ -950,7 +1404,7 @@ void UFastGeoContainer::OnAsyncCreatePhysicsStateBegin_GameThread(TSet<UObject*>
 bool UFastGeoContainer::OnAsyncCreatePhysicsState(const UE::FTimeout& Timeout)
 {
 	TRACE_CPUPROFILER_EVENT_SCOPE(UFastGeoContainer::OnAsyncCreatePhysicsState);
-	check((PhysicsStateCreation == EPhysicsStateCreation::Creating) && !PendingCreate.PhysicsState.IsCompleted());
+	check((PhysicsState == EPhysicsState::Creating) && !PendingCreate.PhysicsState.IsCompleted());
 	
 	for (int Index = PendingCreate.PhysicsState.TotalNumProcessed; Index < PendingCreate.PhysicsState.ComponentsToProcess->Num(); ++Index)
 	{
@@ -971,8 +1425,8 @@ void UFastGeoContainer::OnAsyncCreatePhysicsStateEnd_GameThread()
 	TRACE_CPUPROFILER_EVENT_SCOPE(UFastGeoContainer::OnAsyncCreatePhysicsStateEnd_GameThread);
 	check(PendingCreate.PhysicsState.ComponentsToProcess && PendingCreate.PhysicsState.IsCompleted());
 	PendingCreate.PhysicsState.Reset();
-	check(PhysicsStateCreation == EPhysicsStateCreation::Creating);
-	PhysicsStateCreation = EPhysicsStateCreation::Created;
+	check(PhysicsState == EPhysicsState::Creating);
+	PhysicsState = EPhysicsState::Created;
 	
 	for (FFastGeoComponent* Component : CollisionComponents)
 	{
@@ -980,15 +1434,21 @@ void UFastGeoContainer::OnAsyncCreatePhysicsStateEnd_GameThread()
 	}
 
 	RegisterBodyInstances();
+
+	// Physics work complete - check if all registration systems are done
+	if (!HasAnyPendingCreateTasks())
+	{
+		OnRegisterCompleted();
+	}
 }
 
 void UFastGeoContainer::OnAsyncDestroyPhysicsStateBegin_GameThread(TSet<UObject*>& OutRootedObjects)
 {
 	TRACE_CPUPROFILER_EVENT_SCOPE(UFastGeoContainer::OnAsyncDestroyPhysicsStateBegin_GameThread);
-	check(PhysicsStateCreation == EPhysicsStateCreation::Created);
+	check(PhysicsState == EPhysicsState::Created);
 	check(PendingDestroy.PhysicsState.ComponentsToProcess);
 	check(PendingDestroy.PhysicsState.ComponentsToProcess->Num());
-	PhysicsStateCreation = EPhysicsStateCreation::Destroying;
+	PhysicsState = EPhysicsState::Destroying;
 
 	for (FFastGeoComponent* Component : CollisionComponents)
 	{
@@ -1020,13 +1480,19 @@ void UFastGeoContainer::OnAsyncDestroyPhysicsStateEnd_GameThread()
 	TRACE_CPUPROFILER_EVENT_SCOPE(UFastGeoContainer::OnAsyncDestroyPhysicsStateEnd_GameThread);
 	check(PendingDestroy.PhysicsState.ComponentsToProcess && PendingDestroy.PhysicsState.IsCompleted());
 	PendingDestroy.PhysicsState.Reset();
-	check(PhysicsStateCreation == EPhysicsStateCreation::Destroying);
-	PhysicsStateCreation = EPhysicsStateCreation::NotCreated;
+	check(PhysicsState == EPhysicsState::Destroying);
+	PhysicsState = EPhysicsState::None;
 
 	for (FFastGeoComponent* Component : CollisionComponents)
 	{
 		Component->OnAsyncDestroyPhysicsStateEnd_GameThread();
 	}
+
+	// Physics work complete - check if all unregistration systems are done
+	if (!HasAnyPendingDestroyTasks())
+	{
+		OnUnregisterCompleted();
+	}
 }
 
 void UFastGeoContainer::CollectBodySetupsWithPhysicsMeshesToCreate(TSet<UBodySetup*>& OutBodySetups) const
diff --git a/Engine/Plugins/Experimental/FastGeoStreaming/Source/FastGeoStreaming/Private/FastGeoPrimitiveComponent.cpp b/Engine/Plugins/Experimental/FastGeoStreaming/Source/FastGeoStreaming/Private/FastGeoPrimitiveComponent.cpp
--- a/Engine/Plugins/Experimental/FastGeoStreaming/Source/FastGeoStreaming/Private/FastGeoPrimitiveComponent.cpp
+++ b/Engine/Plugins/Experimental/FastGeoStreaming/Source/FastGeoStreaming/Private/FastGeoPrimitiveComponent.cpp
@@ -513,7 +513,12 @@ bool FFastGeoPrimitiveComponent::ShouldCreateRenderState() const
 
 bool FFastGeoPrimitiveComponent::IsRegistered() const
 {
-	return Super::IsRegistered();
+	// Keep behavior where registering was also considered registered
+	if (UFastGeoContainer* Container = GetOwnerContainer(); ensure(Container))
+	{
+		return Container->IsRegistered() || Container->IsRegistering();
+	}
+	return false;
 }
 
 bool FFastGeoPrimitiveComponent::IsUnreachable() const
diff --git a/Engine/Plugins/Experimental/FastGeoStreaming/Source/FastGeoStreaming/Private/FastGeoWorldSubsystem.cpp b/Engine/Plugins/Experimental/FastGeoStreaming/Source/FastGeoStreaming/Private/FastGeoWorldSubsystem.cpp
--- a/Engine/Plugins/Experimental/FastGeoStreaming/Source/FastGeoStreaming/Private/FastGeoWorldSubsystem.cpp
+++ b/Engine/Plugins/Experimental/FastGeoStreaming/Source/FastGeoStreaming/Private/FastGeoWorldSubsystem.cpp
@@ -300,25 +300,38 @@ void UFastGeoWorldSubsystem::OnLevelStartedRemoveFromWorld(UWorld* World, const
 
 void UFastGeoWorldSubsystem::OnWorldCleanup(UWorld* World, bool bSessionEnded, bool bCleanupResources)
 {
-	if (World != GetWorld())
+	if (World == GetWorld())
 	{
-		return;
+		DrainPendingDestroyContainers();
 	}
+}
 
-	const bool bWaitForCompletion = true;
-	for (const TObjectPtr<UFastGeoContainer>& Container : PendingDestroyContainers)
+void UFastGeoWorldSubsystem::DrainPendingDestroyContainers()
+{
+	check(IsInGameThread());
+
+	static constexpr bool bWaitForCompletion = true;
+
+	while (!PendingDestroyContainers.IsEmpty())
 	{
-		Container->Tick(bWaitForCompletion);
-		check(!Container->HasAnyPendingDestroyTasks());
+		TSet<TObjectPtr<UFastGeoContainer>> Containers = MoveTemp(PendingDestroyContainers);
+		for (const TObjectPtr<UFastGeoContainer>& Container : Containers)
+		{
+			if (IsValid(Container))
+			{
+				Container->Tick(bWaitForCompletion);
+				check(Container->IsFullyUnregistered());
+			}
+		}
 	}
-	PendingDestroyContainers.Reset();
+	check(PendingDestroyContainers.IsEmpty());
 }
 
 void UFastGeoWorldSubsystem::OnGlobalComponentReregisterContextCreated()
 {
 	check(ReregisteringFastGeoContainers.IsEmpty());
 
-	WaitForAllAsyncJobsCompletion();
+	WaitForAllPendingWorkCompletion();
 	
 	TGuardValue<bool> GuardValue(bWaitingForCompletion, true);
 	UWorld* World = GetWorld();
@@ -332,12 +345,12 @@ void UFastGeoWorldSubsystem::OnGlobalComponentReregisterContextCreated()
 		}
 	}
 
-	WaitForAllAsyncJobsCompletion();
+	WaitForAllPendingWorkCompletion();
 }
 
 void UFastGeoWorldSubsystem::OnGlobalComponentReregisterContextDestroyed()
 {
-	WaitForAllAsyncJobsCompletion();
+	WaitForAllPendingWorkCompletion();
 
 	TGuardValue<bool> GuardValue(bWaitingForCompletion, true);
 	for (const TWeakObjectPtr<UFastGeoContainer>& WeakFastGeo : ReregisteringFastGeoContainers)
@@ -349,7 +362,7 @@ void UFastGeoWorldSubsystem::OnGlobalComponentReregisterContextDestroyed()
 		}
 	}
 
-	WaitForAllAsyncJobsCompletion();
+	WaitForAllPendingWorkCompletion();
 
 	ReregisteringFastGeoContainers.Empty();
 }
@@ -359,10 +372,12 @@ void UFastGeoWorldSubsystem::OnPreRecreateScene(UWorld* World)
 {
 	if (World == GetWorld())
 	{
-		// Make sure to flush AsyncRenderStateJobQueue, as it holds a raw pointer to the scene until it completes
-		const bool bWaitForCompletion = true;
-		ProcessAsyncRenderStateJobs(bWaitForCompletion);
-		check(!AsyncRenderStateJobQueue.IsValid());
+		// Ensure all FastGeo work (sync + async) is fully drained before the scene is recreated.
+		// In sync mode, render state work is advanced via UFastGeoContainer::Tick(),
+		// not through the async job queue, so ProcessAsyncRenderStateJobs() alone is insufficient.
+		// This guarantees no FastGeo render/physics jobs still reference the old scene.
+		WaitForAllPendingWorkCompletion();
+		check(!AsyncRenderStateJobQueue.IsValid() || AsyncRenderStateJobQueue->IsCompleted());
 	}
 }
 #endif // WITH_EDITOR
@@ -380,7 +395,7 @@ void UFastGeoWorldSubsystem::OnLevelComponentsUpdated(UWorld* World, ULevel* Lev
 
 		FastGeo->Register();
 
-		const bool bWaitForCompletion = true;
+		static constexpr bool bWaitForCompletion = true;
 		FastGeo->Tick(bWaitForCompletion);
 	}
 }
@@ -400,13 +415,12 @@ void UFastGeoWorldSubsystem::OnLevelComponentsCleared(UWorld* World, ULevel* Lev
 		{
 			FastGeo->Unregister();
 
-			const bool bWaitForCompletion = true;
+			static constexpr bool bWaitForCompletion = true;
 			FastGeo->Tick(bWaitForCompletion);
 		}
 		else
 		{
-			check(!FastGeo->IsRegistered());
-			check(!FastGeo->HasAnyPendingTasks());
+			check(FastGeo->IsFullyUnregistered());
 		}
 	}
 }
@@ -456,18 +470,52 @@ void UFastGeoWorldSubsystem::OnUpdateLevelStreaming()
 	if (GetWorld()->GetShouldForceUnloadStreamingLevels())
 	{
 		TRACE_CPUPROFILER_EVENT_SCOPE(UFastGeoWorldSubsystem::OnUpdateLevelStreaming);
-		WaitForAllAsyncJobsCompletion();
+		WaitForAllPendingWorkCompletion();
 	}
 }
 
-void UFastGeoWorldSubsystem::WaitForAllAsyncJobsCompletion()
+void UFastGeoWorldSubsystem::WaitForAllPendingWorkCompletion()
 {
-	const bool bWaitForCompletion = true;
-	ProcessAsyncRenderStateJobs(bWaitForCompletion);
-	if (FPhysScene* PhysScene = GetWorld()->GetPhysicsScene())
+	check(IsInGameThread());
+
+	static constexpr bool bWaitForCompletion = true;
+
+#if DO_CHECK
+	int32 SafetyMaxIteration = 0;
+#endif
+	
+	// Converge all FastGeo containers in this world to their target states
+	UWorld* World = GetWorld();
+	while (true)
 	{
-		PhysScene->ProcessAsyncPhysicsStateJobs(bWaitForCompletion);
+		for (TObjectIterator<UFastGeoContainer> It(RF_ClassDefaultObject, true, EInternalObjectFlags::Garbage); It; ++It)
+		{
+			UFastGeoContainer* FastGeo = *It;
+			if (FastGeo->GetWorld() == World)
+			{
+				FastGeo->Tick(bWaitForCompletion);
+			}
+		}
+
+		// Drain recreate queue
+		ProcessPendingRecreate();
+		if (ComponentsPendingRecreate.IsEmpty())
+		{
+			break;
+		}
+
+#if DO_CHECK
+		checkf(++SafetyMaxIteration < 1000, TEXT("WaitForAllPendingWorkCompletion() not converging."));
+#endif
 	}
+
+	DrainPendingDestroyContainers();
+
+	FPhysScene* PhysScene = World ? World->GetPhysicsScene() : nullptr;
+	check(!PhysScene || !PhysScene->HasAsyncPhysicsStateJobs());
+	check(!AsyncRenderStateJobQueue.IsValid() || AsyncRenderStateJobQueue->IsCompleted());
+	check(ComponentsPendingRecreate.IsEmpty());
+	check(PendingDestroyContainers.IsEmpty());
 }
 
 void UFastGeoWorldSubsystem::AddToComponentsPendingRecreate(FFastGeoComponent* InComponentPendingRecreate)
@@ -587,7 +635,14 @@ void UFastGeoWorldSubsystem::Tick(float DeltaTime)
 	
 	for (auto It = PendingDestroyContainers.CreateIterator(); It; ++It)
 	{
-		if (!(*It)->HasAnyPendingDestroyTasks())
+		UFastGeoContainer* Container = *It;
+		bool bShouldRemove = !IsValid(Container) || Container->IsFullyUnregistered();
+		if (!bShouldRemove)
+		{
+			Container->Tick();
+			bShouldRemove = Container->IsFullyUnregistered();
+		}
+		if (bShouldRemove)
 		{
 			It.RemoveCurrent();
 		}
@@ -651,7 +706,7 @@ void UFastGeoWorldSubsystem::PushAsyncDestroyContainer(UFastGeoContainer* FastGe
 
 	FastGeo->Unregister();
 
-	if (FastGeo->HasAnyPendingDestroyTasks())
+	if (!FastGeo->IsFullyUnregistered())
 	{
 		PendingDestroyContainers.Add(FastGeo);
 	}
diff --git a/Engine/Source/Runtime/Engine/Classes/Engine/World.h b/Engine/Source/Runtime/Engine/Classes/Engine/World.h
--- a/Engine/Source/Runtime/Engine/Classes/Engine/World.h
+++ b/Engine/Source/Runtime/Engine/Classes/Engine/World.h
@@ -3176,6 +3176,13 @@ public:
 	{
 		RemoveFromWorld(Level, bAllowIncrementalRemoval, TOptional<const UE::FTimeout>(), TransactionId, OwningLevelStreaming);
 	}
+
+	/** Returns current timeout used by AddToWorld */
+	UE_API UE::FTimeout GetAddToWorldTimeout() const;
+
+	/** Returns current timeout used by RemoveFromWorld */
+	UE_API UE::FTimeout GetRemoveFromWorldTimeout() const;
+
 	/**
 	 * Updates sub-levels (load/unload/show/hide) using streaming levels current state
 	 *
@@ -4388,6 +4395,9 @@ public:
 
 	//PSC Pooling END
 	FWorldSubsystemCollection SubsystemCollection;
+
+	TOptional<UE::FTimeout> AddToWorldTimeout;
+	TOptional<UE::FTimeout> RemoveFromWorldTimeout;
 };
 
 /** Global UWorld pointer. Use of this pointer should be avoided whenever possible. */
diff --git a/Engine/Source/Runtime/Engine/Private/PhysicsEngine/Experimental/PhysScene_Chaos.cpp b/Engine/Source/Runtime/Engine/Private/PhysicsEngine/Experimental/PhysScene_Chaos.cpp
--- a/Engine/Source/Runtime/Engine/Private/PhysicsEngine/Experimental/PhysScene_Chaos.cpp
+++ b/Engine/Source/Runtime/Engine/Private/PhysicsEngine/Experimental/PhysScene_Chaos.cpp
@@ -1945,6 +1945,12 @@ void FPhysScene_Chaos::ProcessAsyncPhysicsStateJobs(bool bWaitForCompletion)
 	}
 }
 
+bool FPhysScene_Chaos::HasAsyncPhysicsStateJobs() const
+{
+	check(IsInGameThread());
+	return (AsyncPhysicsStateJobQueue.IsValid() && !AsyncPhysicsStateJobQueue->IsCompleted());
+}
+
 // Collect the actors and transforms of all the bodies we have to move, and process them in bulk
 // to avoid locks in the Spatial Acceleration and the Solver's Dirty Proxy systems.
 void FPhysScene_Chaos::UpdateKinematicsOnDeferredSkelMeshes()
diff --git a/Engine/Source/Runtime/Engine/Private/World.cpp b/Engine/Source/Runtime/Engine/Private/World.cpp
--- a/Engine/Source/Runtime/Engine/Private/World.cpp
+++ b/Engine/Source/Runtime/Engine/Private/World.cpp
@@ -266,6 +266,18 @@ FAutoConsoleVariableRef CVarShouldInitWorldReleaseScene(
 	TEXT("Set to false to revert to earlier behavior (causes leaking Scene and associated resources in certain situations)."),
 	ECVF_Default);
 
+UE::FTimeout UWorld::GetAddToWorldTimeout() const
+{
+	check(IsInGameThread());
+	return AddToWorldTimeout.Get(UE::FTimeout::Never());
+}
+
+UE::FTimeout UWorld::GetRemoveFromWorldTimeout() const
+{
+	check(IsInGameThread());
+	return RemoveFromWorldTimeout.Get(UE::FTimeout::Never());
+}
+
 // This cvar activates detection of incorrectly unregistered components after a world is cleaned-up.
 // This is important since components left as registered/renderstatecreated after their world is cleaned up will 
 // crash in FGlobalComponentRecreateRenderStateContext unless a GC happens before the recreate.
@@ -3712,6 +3724,9 @@ void UWorld::AddToWorld(ULevel* Level, const FTransform& LevelTransform, bool bC
 		bExecuteNextStep = !IsTimeLimitExceeded(Timeout, TEXT("computing time limit"), Level);
 	}
 
+	AddToWorldTimeout = Timeout;
+	ON_SCOPE_EXIT { AddToWorldTimeout.Reset(); };
+
 	// Don't make this level visible if it's currently being made invisible
 	if( bExecuteNextStep && !bIsLevelMakingVisible && !IsLevelMakingInvisible(Level) )
 	{
@@ -4108,6 +4123,9 @@ void UWorld::RemoveFromWorld(ULevel* Level, bool bAllowIncrementalRemoval, const
 			: UE::FTimeout(FMath::Max<double>(0.0, GLevelStreamingUnregisterComponentsTimeLimit - GRemoveFromWorldTimeCumul) / 1000.0);
 	}
 
+	RemoveFromWorldTimeout = Timeout;
+	ON_SCOPE_EXIT { RemoveFromWorldTimeout.Reset(); };
+
 	FScopeCycleCounterUObject ScopeContext(Level);
 	check(IsValid(Level));
 	check(!Level->IsUnreachable());
diff --git a/Engine/Source/Runtime/Engine/Public/Physics/Experimental/PhysScene_Chaos.h b/Engine/Source/Runtime/Engine/Public/Physics/Experimental/PhysScene_Chaos.h
--- a/Engine/Source/Runtime/Engine/Public/Physics/Experimental/PhysScene_Chaos.h
+++ b/Engine/Source/Runtime/Engine/Public/Physics/Experimental/PhysScene_Chaos.h
@@ -284,6 +284,7 @@ public:
 	ENGINE_API bool PushAsyncDestroyPhysicsStateJob(IAsyncPhysicsStateProcessor* Processor);
 	ENGINE_API void RemoveAsyncPhysicsStateJobs(IAsyncPhysicsStateProcessor* Processor);
 	ENGINE_API void ProcessAsyncPhysicsStateJobs(bool bWaitForCompletion = false);
+	ENGINE_API bool HasAsyncPhysicsStateJobs() const;
 
 	UChaosEventRelay* GetChaosEventRelay() const { return ChaosEventRelay.Get(); }
 
