[FastGeo] Thread safety fixes for async render state creation
- PIE forces sync mode (GT + ParallelGT). Cooked builds use async mode by default (generic worker threads).
- FastGeo requires FSimpleStreamableAssetManager to be enabled (ensured at module startup)
- Fixes to make async render state creation safe from generic workers:
. Cache world/scene/physics/subsystem pointers on GT before workers run (FFastGeoAsyncContext), because UWorld subsystem lookups, time reads, and physics scene access are not safe from generic workers
. Pre-call ConditionalPostLoad on all referenced materials on GT, because workers cannot trigger PostLoad on materials that haven't completed it yet
. Pre-warm material relevance cache on GT (FAtomicCachedMaterialRelevance on UMaterialInterface), because GetRelevance_Concurrent reads from GameThreadShaderMap which is not safe from generic workers without a primed cache
. Both steps above require knowing the full set of referenced materials upfront: for cooked builds, the list is pre-computed at cook time (hard refs for base materials, soft refs for Nanite overrides that may be stripped); for PIE/transient containers, it is built on demand via a virtual ForEachMaterial on FFastGeoComponent
#pf 69a9a441f874567e9d587bd4

#Synced-CL 51093031

diff --git a/Engine/Plugins/Experimental/FastGeoStreaming/Source/FastGeoStreaming/Internal/FastGeoComponent.h b/Engine/Plugins/Experimental/FastGeoStreaming/Source/FastGeoStreaming/Internal/FastGeoComponent.h
--- a/Engine/Plugins/Experimental/FastGeoStreaming/Source/FastGeoStreaming/Internal/FastGeoComponent.h
+++ b/Engine/Plugins/Experimental/FastGeoStreaming/Source/FastGeoStreaming/Internal/FastGeoComponent.h
@@ -8,6 +8,7 @@
 #include "IFastGeoElement.h"
 #include "FastGeoElementType.h"
 #include "Physics/Experimental/AsyncPhysicsStateProcessorInterface.h"
+#include "PhysicsInterfaceDeclaresCore.h"
 #include "Components/ComponentInterfaces.h"
 #include "Components/SceneComponent.h"
 #include "Async/TaskGraphInterfaces.h"
@@ -49,20 +50,37 @@ public:
 #if WITH_EDITOR
 	virtual void InitializeFromComponent(UActorComponent* Component);
 #endif
-	int32 GetComponentIndex() const { return ComponentIndex; }
-	void MarkRenderStateDirty(bool bEvenIfNotCreated);
-	bool IsRenderStateDirty() const;
-	bool IsRenderStateDelayed() const;
-	bool IsRenderStateCreated() const;
-	virtual UBodySetup* GetBodySetup() const { return nullptr; }
-	virtual bool IsCollisionEnabled() const { return false; }
 	virtual void Serialize(FArchive& Ar);
 	virtual void InitializeDynamicProperties() {}
 	virtual void ApplyWorldTransform(const FTransform& InTransform);
+	UWorld* GetWorld() const;
+	FSceneInterface* GetScene() const;
+	FPhysScene* GetPhysicsScene() const;
+	float GetWorldTimeSeconds() const;
+	UFastGeoContainer* GetOwnerContainer() const;
+	FFastGeoComponentCluster* GetOwnerComponentCluster() const;
+	int32 GetComponentIndex() const { return ComponentIndex; }
+	FLinearColor GetDebugColor() const;
+
+	// Enumerates all materials this component references, including Nanite overrides.
+	// bIsNaniteOverride is true for materials obtained via GetNaniteOverride().
+	virtual void ForEachMaterial(TFunctionRef<void(UMaterialInterface*, bool bIsNaniteOverride)> Func) const {}
+
+	// Render interface
+	virtual void UpdateVisibility() {}
 	virtual bool ShouldComponentAddToRenderScene() const;
 	virtual void CreateRenderState(FRegisterComponentContext* Context) {}
 	virtual void DestroyRenderState(FFastGeoDestroyRenderStateContext* Context) {}
-	virtual void UpdateVisibility() {}
+	void MarkRenderStateDirty(bool bEvenIfNotCreated);
+	bool IsRenderStateDirty() const;
+	bool IsRenderStateDelayed() const;
+	bool IsRenderStateCreated() const;
+
+	// Physics interface
+	virtual bool IsCollisionEnabled() const { return false; }
+	virtual UBodySetup* GetBodySetup() const { return nullptr; }
+
+	// Async physics state creation/destruction
 	virtual void OnAsyncCreatePhysicsStateBegin_GameThread() {}
 	virtual void OnAsyncCreatePhysicsState();
 	virtual void OnAsyncCreatePhysicsStateEnd_GameThread();
@@ -72,18 +90,13 @@ public:
 	virtual void GetBodyInstances(TArray<FBodyInstance*>& OutBodyInstances) {}
 
 #if !WITH_EDITOR && UE_WITH_PSO_PRECACHING
+	// PSO Precaching interface
 	virtual bool ShouldPrecachePSOs() const;
 	virtual void PrecachePSOs_Concurrent() {}
 	virtual FPSOPrecacheComponentData& GetPSOPrecacheComponentData() = 0;
 	virtual const FPSOPrecacheComponentData& GetPSOPrecacheComponentData() const = 0;
 #endif
 
-	FFastGeoComponentCluster* GetOwnerComponentCluster() const;
-	UFastGeoContainer* GetOwnerContainer() const;
-	UWorld* GetWorld() const;
-	FSceneInterface* GetScene() const;
-	FLinearColor GetDebugColor() const;
-
 protected:
 	// Persistent Data
 	int32 ComponentIndex = INDEX_NONE;
@@ -96,6 +109,7 @@ protected:
 	TUniquePtr<FRWLock> Lock;
 
 #ifdef FASTGEO_DEBUG_PHYSICS
+	// Physics state
 	enum EPhysicsStateCreation
 	{
 		NotCreated,
@@ -107,6 +121,7 @@ protected:
 	EPhysicsStateCreation PhysicsStateCreation = EPhysicsStateCreation::NotCreated;
 #endif
 
+	// Render state
 	enum class EProxyCreationState : uint8
 	{
 		None, // Constructed/Initialized
@@ -121,6 +136,5 @@ protected:
 private:
 	void SetOwnerComponentCluster(FFastGeoComponentCluster* InOwner);
 	friend class FFastGeoComponentCluster;
-
 	friend FArchive& operator<<(FArchive& Ar, FFastGeoComponent& Component);
 };
\ No newline at end of file
diff --git a/Engine/Plugins/Experimental/FastGeoStreaming/Source/FastGeoStreaming/Internal/FastGeoContainer.h b/Engine/Plugins/Experimental/FastGeoStreaming/Source/FastGeoStreaming/Internal/FastGeoContainer.h
--- a/Engine/Plugins/Experimental/FastGeoStreaming/Source/FastGeoStreaming/Internal/FastGeoContainer.h
+++ b/Engine/Plugins/Experimental/FastGeoStreaming/Source/FastGeoStreaming/Internal/FastGeoContainer.h
@@ -6,10 +6,12 @@
 #include "Physics/Experimental/AsyncPhysicsStateProcessorInterface.h"
 #include "PhysicsEngine/PhysicsBodyInstanceOwnerInterface.h"
 #include "PhysicsEngine/BodyInstance.h"
+#include "PhysicsInterfaceDeclaresCore.h"
 #include "AI/Navigation/NavigationElement.h"
 #include "FastGeoComponentCluster.h"
 #include "FastGeoHLOD.h"
 #include "UObject/ObjectPtr.h"
+#include "UObject/SoftObjectPtr.h"
 
 #if WITH_EDITOR
 #include "UObject/ObjectSaveContext.h"
@@ -23,11 +25,71 @@ class URuntimeVirtualTexture;
 class UStaticMesh;
 class UMaterialInterface;
 class UNavigationSystemBase;
+class UFastGeoWorldSubsystem;
 class FFastGeoComponent;
 class FFastGeoComponentCluster;
 class FChaosUserDefinedEntity;
 class AFastGeoSurrogateActor;
 
+/**
+ * Cached world/scene/subsystem/physics pointers.
+ * Populated on GT before workers run.
+ * Workers read from this instead of performing lookups not safe from generic workers
+ * (e.g., subsystem collection access, world time reads, physics scene access).
+ * Reset on unregister. Re-populated on re-register.
+ * Lifetime: World/Scene/Subsystem all outlive the container's registered state.
+ */
+struct FFastGeoAsyncContext
+{
+	void Initialize(UFastGeoContainer& Container);
+	void Reset() { *this = {}; }
+	bool IsInitialized() const { return bInitialized; }
+
+	UWorld* GetWorld() const { check(bInitialized); return World; }
+	FSceneInterface* GetScene() const { check(bInitialized); return Scene; }
+	FPhysScene* GetPhysicsScene() const { check(bInitialized); return PhysicsScene; }
+	UFastGeoWorldSubsystem* GetWorldSubsystem() const { check(bInitialized); return WorldSubsystem; }
+	float GetTimeSeconds() const { check(bInitialized); return TimeSeconds; }
+	void UpdateTimeSeconds(float InTime) { check(bInitialized); TimeSeconds = InTime; }
+
+private:
+	bool bInitialized = false;
+	UWorld* World = nullptr;
+	FSceneInterface* Scene = nullptr;
+	FPhysScene* PhysicsScene = nullptr;
+	UFastGeoWorldSubsystem* WorldSubsystem = nullptr;
+	float TimeSeconds = 0.0f;
+};
+
+/**
+ * Calls ConditionalPostLoad on all materials referenced by the container.
+ * Ensures RF_NeedPostLoad is cleared before workers access materials.
+ * Permanent - never reset. RF_NeedPostLoad is cleared once.
+ */
+struct FFastGeoMaterialPostLoadInit
+{
+	void Initialize(UFastGeoContainer& Container);
+	bool IsInitialized() const { return bInitialized; }
+
+private:
+	bool bInitialized = false;
+};
+
+/**
+ * Warms UMaterialInterface::CachedRelevance for all materials + default material.
+ * Workers use the cached relevance instead of reading GameThreadShaderMap.
+ * Permanent - never reset. Cache persists on UMaterialInterface.
+ * Applies to async mode only.
+ */
+struct FFastGeoMaterialRelevanceInit
+{
+	void Initialize(UFastGeoContainer& Container);
+	bool IsInitialized() const { return bInitialized; }
+
+private:
+	bool bInitialized = false;
+};
+
 struct FFastGeoComponentClusterPendingState
 {
 	struct FRenderState
@@ -170,15 +232,17 @@ public:
 
 	void InitializeDynamicProperties(bool bInitForPlay = false);
 
-	void OnPrecachePSOsBegin_GameThread() {}
+	void OnPrecachePSOsBegin_GameThread();
+	void OnPrecachePSOs_Concurrent();
 	void OnPrecachePSOsEnd_GameThread() {}
+
 	void OnCreateRenderStateBegin_GameThread();
-	void OnDestroyRenderStateBegin_GameThread();
-	void OnCreateRenderStateEnd_GameThread();
-	void OnDestroyRenderStateEnd_GameThread();
-	void OnPrecachePSOs_Concurrent();
 	void OnCreateRenderState_Concurrent();
+	void OnCreateRenderStateEnd_GameThread();
+	
+	void OnDestroyRenderStateBegin_GameThread();
 	void OnDestroyRenderState_Concurrent();
+	void OnDestroyRenderStateEnd_GameThread();
 
 	void OnCreatePhysicsStateBegin_GameThread();
 	void OnDestroyPhysicsStateBegin_GameThread();
@@ -231,6 +295,18 @@ protected:
 
 private:
 
+	UFastGeoWorldSubsystem* GetWorldSubsystem() const;
+	FSceneInterface* GetScene() const;
+	FPhysScene* GetPhysicsScene() const;
+	float GetWorldTimeSeconds() const;
+
+	void InitializePrerequisites();
+	void ForEachMaterialAsset(TFunctionRef<void(UMaterialInterface*)> Func);
+	void BuildMaterialAssets();
+#if !WITH_EDITOR
+	void CacheMaterialAssets();
+#endif
+
 	void StartCreatePhysicsStateWork();
 	void StartDestroyPhysicsStateWork();
 	void TickPhysicsWork(bool bWaitForCompletion);
@@ -361,13 +437,29 @@ private:
 	UPROPERTY()
 	TArray<TObjectPtr<UObject>> Assets;
 
+	// Pre-computed unique material lists, built at cook time in PreSave.
+	// Base materials (OverrideMaterials, mesh slots, decal, light function) are hard refs
+	UPROPERTY()
+	TArray<TObjectPtr<UMaterialInterface>> MaterialAssetHardRefs;
+
+	// Nanite override materials are soft refs - may be stripped on non-Nanite platforms
+	// (FMaterialOverrideNanite::Serialize can null them during cook).
+	// Resolved in CacheMaterialAssets when loading (Serialize) cooked container.
+	UPROPERTY()
+	TArray<TSoftObjectPtr<UMaterialInterface>> MaterialAssetNaniteOverrideSoftRefs;
+
 	UPROPERTY()
 	TObjectPtr<AFastGeoSurrogateActor> SurrogateActor;
 
 #if !WITH_EDITOR && UE_WITH_PSO_PRECACHING
-	bool bPrecachedPSOs = false;
+	bool bProcessedPSOPrecaching = false;
 #endif
 
+	// Resolved material cache (transient).
+	UPROPERTY(Transient)
+	TArray<TObjectPtr<UMaterialInterface>> CachedMaterialAssets;
+	bool bMaterialAssetsCached = false;
+
 	// Transient data
 	TArray<IPrimitiveComponent*> PrimitiveComponents;
 	TArray<IStaticMeshComponent*> StaticMeshComponents;
@@ -377,12 +469,22 @@ private:
 	bool bIsBeingAsyncDestroyed = false;
 	int32 RegistrationEpoch = 0;
 
+	// Prerequisites: must be initialized on GT before starting render/physics work.
+	// Workers read from AsyncContext instead of performing lookups not safe from generic workers.
+	FFastGeoAsyncContext AsyncContext;
+	FFastGeoMaterialPostLoadInit MaterialPostLoadInit;
+	FFastGeoMaterialRelevanceInit MaterialRelevanceInit;
+
 	FRenderCommandFence DestroyFence;
 
 	FOnUnregistered OnUnregistered;
 
+	friend class UFastGeoWorldSubsystem;
 	friend class UFastGeoWorldPartitionRuntimeCellTransformer;
 	friend class FFastGeoGatherFastGeoContainerAssetRefsArchive;
+	friend class FFastGeoComponent;
 	friend struct FFastGeoRegisteredComponent;
-	friend class UFastGeoWorldSubsystem;
+	friend struct FFastGeoAsyncContext;
+	friend struct FFastGeoMaterialPostLoadInit;
+	friend struct FFastGeoMaterialRelevanceInit;
 };
\ No newline at end of file
diff --git a/Engine/Plugins/Experimental/FastGeoStreaming/Source/FastGeoStreaming/Internal/FastGeoDecalComponent.h b/Engine/Plugins/Experimental/FastGeoStreaming/Source/FastGeoStreaming/Internal/FastGeoDecalComponent.h
--- a/Engine/Plugins/Experimental/FastGeoStreaming/Source/FastGeoStreaming/Internal/FastGeoDecalComponent.h
+++ b/Engine/Plugins/Experimental/FastGeoStreaming/Source/FastGeoStreaming/Internal/FastGeoDecalComponent.h
@@ -21,6 +21,7 @@ public:
 	virtual ~FFastGeoDecalComponent() = default;
 
 	//~ Begin FFastGeoComponent interface
+	virtual void ForEachMaterial(TFunctionRef<void(UMaterialInterface*, bool bIsNaniteOverride)> Func) const override;
 	virtual void Serialize(FArchive& Ar) override;
 	virtual void ApplyWorldTransform(const FTransform& InTransform) override;
 	virtual bool ShouldComponentAddToRenderScene() const override;
diff --git a/Engine/Plugins/Experimental/FastGeoStreaming/Source/FastGeoStreaming/Internal/FastGeoLocalLightComponent.h b/Engine/Plugins/Experimental/FastGeoStreaming/Source/FastGeoStreaming/Internal/FastGeoLocalLightComponent.h
--- a/Engine/Plugins/Experimental/FastGeoStreaming/Source/FastGeoStreaming/Internal/FastGeoLocalLightComponent.h
+++ b/Engine/Plugins/Experimental/FastGeoStreaming/Source/FastGeoStreaming/Internal/FastGeoLocalLightComponent.h
@@ -20,6 +20,7 @@ public:
 	virtual ~FFastGeoLocalLightComponent() = default;
 
 	//~ Begin FFastGeoComponent interface
+	virtual void ForEachMaterial(TFunctionRef<void(UMaterialInterface*, bool bIsNaniteOverride)> Func) const override;
 	virtual void Serialize(FArchive& Ar) override;
 	virtual bool ShouldComponentAddToRenderScene() const override;
 	virtual void CreateRenderState(FRegisterComponentContext* Context) override;
diff --git a/Engine/Plugins/Experimental/FastGeoStreaming/Source/FastGeoStreaming/Internal/FastGeoSkinnedMeshComponent.h b/Engine/Plugins/Experimental/FastGeoStreaming/Source/FastGeoStreaming/Internal/FastGeoSkinnedMeshComponent.h
--- a/Engine/Plugins/Experimental/FastGeoStreaming/Source/FastGeoStreaming/Internal/FastGeoSkinnedMeshComponent.h
+++ b/Engine/Plugins/Experimental/FastGeoStreaming/Source/FastGeoStreaming/Internal/FastGeoSkinnedMeshComponent.h
@@ -30,8 +30,9 @@ public:
 	virtual void InitializeFromComponent(UActorComponent* Component) override;
 #endif
 	//~ End FFastGeoComponent interface
-	
+
 	//~ Begin FFastGeoPrimitiveComponent interface
+	virtual void ForEachMaterial(TFunctionRef<void(UMaterialInterface*, bool bIsNaniteOverride)> Func) const override;
 	virtual UMaterialInterface* GetMaterial(int32 MaterialIndex) const override;
 	virtual int32 GetNumMaterials() const override;
 	virtual void GetUsedMaterials(TArray<UMaterialInterface*>& OutMaterials, bool bGetDebugMaterials = false) const override;
diff --git a/Engine/Plugins/Experimental/FastGeoStreaming/Source/FastGeoStreaming/Internal/FastGeoStaticMeshComponent.h b/Engine/Plugins/Experimental/FastGeoStreaming/Source/FastGeoStreaming/Internal/FastGeoStaticMeshComponent.h
--- a/Engine/Plugins/Experimental/FastGeoStreaming/Source/FastGeoStreaming/Internal/FastGeoStaticMeshComponent.h
+++ b/Engine/Plugins/Experimental/FastGeoStreaming/Source/FastGeoStreaming/Internal/FastGeoStaticMeshComponent.h
@@ -76,6 +76,7 @@ public:
 	virtual FBox GetNavigationBounds() const override;
 	virtual void GetNavigationData(FNavigationRelevantData& Data) const override;
 	virtual bool DoCustomNavigableGeometryExport(FNavigableGeometryExport& GeomExport) const override;
+	virtual void ForEachMaterial(TFunctionRef<void(UMaterialInterface*, bool bIsNaniteOverride)> Func) const override;
 	virtual UMaterialInterface* GetMaterial(int32 MaterialIndex) const override;
 	virtual int32 GetNumMaterials() const override;
 	virtual void GetUsedMaterials(TArray<UMaterialInterface*>& OutMaterials, bool bGetDebugMaterials = false) const override;
diff --git a/Engine/Plugins/Experimental/FastGeoStreaming/Source/FastGeoStreaming/Private/FastGeoComponent.cpp b/Engine/Plugins/Experimental/FastGeoStreaming/Source/FastGeoStreaming/Private/FastGeoComponent.cpp
--- a/Engine/Plugins/Experimental/FastGeoStreaming/Source/FastGeoStreaming/Private/FastGeoComponent.cpp
+++ b/Engine/Plugins/Experimental/FastGeoStreaming/Source/FastGeoStreaming/Private/FastGeoComponent.cpp
@@ -71,19 +71,22 @@ UFastGeoContainer* FFastGeoComponent::GetOwnerContainer() const
 
 UWorld* FFastGeoComponent::GetWorld() const
 {
-	check(GetOwnerContainer());
-	check(GetOwnerContainer()->GetLevel());
-
-	return GetOwnerContainer()->GetLevel()->OwningWorld;
+	return GetOwnerContainer()->GetWorld();
 }
 
 FSceneInterface* FFastGeoComponent::GetScene() const
 {
-	if (UWorld* World = GetWorld())
-	{
-		return World->Scene;
-	}
-	return nullptr;
+	return GetOwnerContainer()->GetScene();
+}
+
+FPhysScene* FFastGeoComponent::GetPhysicsScene() const
+{
+	return GetOwnerContainer()->GetPhysicsScene();
+}
+
+float FFastGeoComponent::GetWorldTimeSeconds() const
+{
+	return GetOwnerContainer()->GetWorldTimeSeconds();
 }
 
 FLinearColor FFastGeoComponent::GetDebugColor() const
diff --git a/Engine/Plugins/Experimental/FastGeoStreaming/Source/FastGeoStreaming/Private/FastGeoContainer.cpp b/Engine/Plugins/Experimental/FastGeoStreaming/Source/FastGeoStreaming/Private/FastGeoContainer.cpp
--- a/Engine/Plugins/Experimental/FastGeoStreaming/Source/FastGeoStreaming/Private/FastGeoContainer.cpp
+++ b/Engine/Plugins/Experimental/FastGeoStreaming/Source/FastGeoStreaming/Private/FastGeoContainer.cpp
@@ -5,6 +5,7 @@
 #include "FastGeoHLOD.h"
 #include "FastGeoComponent.h"
 #include "FastGeoStreamingModule.h"
+#include "FastGeoSkinnedMeshComponent.h"
 #include "FastGeoSurrogateActor.h"
 #include "FastGeoSurrogateComponent.h"
 #include "FastGeoWorldSubsystem.h"
@@ -12,11 +13,14 @@
 #include "FastGeoLog.h"
 #include "AI/Navigation/NavigationElement.h"
 #include "Async/ParallelFor.h"
+#include "Engine/Engine.h"
 #include "Engine/Level.h"
 #include "Engine/SkeletalMesh.h"
+#include "Engine/StaticMesh.h"
 #include "Engine/Texture.h"
 #include "Engine/TextureLightProfile.h"
 #include "Engine/World.h"
+#include "Materials/MaterialInterface.h"
 #include "Misc/Timeout.h"
 #include "NavigationSystem.h"
 #include "PhysicsEngine/BodySetup.h"
@@ -26,21 +30,15 @@
 #include "Physics/Experimental/PhysScene_Chaos.h"
 #include "ProfilingDebugging/CountersTrace.h"
 #include "PSOPrecache.h"
+#include "PSOPrecacheSettings.h"
+#include "Streaming/SimpleStreamableAssetManager.h"
+#include "SceneInterface.h"
 #include "Templates/Invoke.h"
 #include "WorldPartition/WorldPartitionHelpers.h"
 #include "WorldPartition/WorldPartition.h"
-#include "Engine/Engine.h"
 
 namespace FastGeo::Private
 {
-	// Async mode is disabled by default until issues are fixed
-	static bool GFastGeoAllowAsyncRenderWork = false;
-	static FAutoConsoleVariableRef CVarFastGeoAllowAsyncRenderWork(
-		TEXT("FastGeo.AllowAsyncRenderWork"),
-		GFastGeoAllowAsyncRenderWork,
-		TEXT("Allow FastGeo to perform render state creation/destruction and PSO precaching asynchronously."),
-		ECVF_ReadOnly);
-
 	static int32 GAsyncRenderStateTaskParallelWorkerCount = 1;
 	static FAutoConsoleVariableRef CVarAsyncRenderStateTaskParallelWorkerCount(
 		TEXT("FastGeo.AsyncRenderStateTask.ParallelWorkerCount"),
@@ -290,17 +288,18 @@ void UFastGeoContainer::StartRegisterTransition()
 #if !WITH_EDITOR && UE_WITH_PSO_PRECACHING
 	// Prepare PendingCreate.PSOPrecachingState
 	PendingCreate.PrecachePSOState.Reset();
-	if (!FastGeo::Private::GFastGeoAllowAsyncRenderWork && !bPrecachedPSOs)
+	if (!FFastGeoStreamingModule::IsAsyncRenderWorkAllowed() && !bProcessedPSOPrecaching)
 	{
 		if (IsComponentPSOPrecachingEnabled())
 		{
 			PendingCreate.PrecachePSOState.ComponentsToProcess = PendingCreate.RenderState.ComponentsToProcess;
 			PendingCreate.PrecachePSOState.NumToProcess = PendingCreate.PrecachePSOState.ComponentsToProcess.Num();
 		}
-		bPrecachedPSOs = (PendingCreate.PrecachePSOState.ComponentsToProcess.Num() == 0);
+		bProcessedPSOPrecaching = (PendingCreate.PrecachePSOState.ComponentsToProcess.Num() == 0);
 	}
 #endif
 
+	InitializePrerequisites();
 	StartCreateRenderStateWork();
 	StartCreatePhysicsStateWork();
 
@@ -369,6 +368,8 @@ void UFastGeoContainer::OnUnregisterCompleted()
 	check(RegistrationState == ERegistrationState::Unregistering);
 	RegistrationState = ERegistrationState::Unregistered;
 
+	AsyncContext.Reset();
+
 	OnUnregistered.Broadcast();
 
 	ApplyRegistrationTargetState();
@@ -426,7 +427,7 @@ void UFastGeoContainer::DestroyTransient(UFastGeoContainer* InFastGeo)
 
 	InFastGeo->bIsBeingAsyncDestroyed = true;
 
-	UFastGeoWorldSubsystem* WorldSubsystem = InFastGeo->GetWorld()->GetSubsystemChecked<UFastGeoWorldSubsystem>();
+	UFastGeoWorldSubsystem* WorldSubsystem = InFastGeo->GetWorldSubsystem();
 	WorldSubsystem->PushAsyncDestroyContainer(InFastGeo);
 }
 
@@ -559,23 +560,23 @@ namespace FastGeo::Private
 bool UFastGeoContainer::TickPrecachePSO_Sync(const UE::FTimeout& Timeout)
 {
 #if !WITH_EDITOR && UE_WITH_PSO_PRECACHING
-	if (!bPrecachedPSOs)
+	if (!bProcessedPSOPrecaching)
 	{
 		TRACE_CPUPROFILER_EVENT_SCOPE(UFastGeoContainer::TickPrecachePSO_Sync);
-		check(!FastGeo::Private::GFastGeoAllowAsyncRenderWork);
+		check(!FFastGeoStreamingModule::IsAsyncRenderWorkAllowed());
 
-		bPrecachedPSOs = FastGeo::Private::AdvanceRenderStateBudgeted(PendingCreate.PrecachePSOState, Timeout, [](FFastGeoComponent* Component)
+		bProcessedPSOPrecaching = FastGeo::Private::AdvanceRenderStateBudgeted(PendingCreate.PrecachePSOState, Timeout, [](FFastGeoComponent* Component)
 		{
 			TRACE_CPUPROFILER_EVENT_SCOPE(UFastGeoContainer::PrecachePSOs);
 			Component->PrecachePSOs_Concurrent();
 		});
 
-		if (bPrecachedPSOs)
+		if (bProcessedPSOPrecaching)
 		{
 			PendingCreate.PrecachePSOState.Reset();
 		}
 	}
-	return bPrecachedPSOs;
+	return bProcessedPSOPrecaching;
 #endif
 	return true;
 }
@@ -584,7 +585,7 @@ bool UFastGeoContainer::TickCreateRenderState_Sync(const UE::FTimeout& Timeout)
 {
 	TRACE_CPUPROFILER_EVENT_SCOPE(UFastGeoContainer::TickCreateRenderState_Sync);
 	check(IsInGameThread());
-	check(!FastGeo::Private::GFastGeoAllowAsyncRenderWork);
+	check(!FFastGeoStreamingModule::IsAsyncRenderWorkAllowed());
 
 	if (PendingCreate.RenderState.ComponentsToProcess.Num() > 0 && PendingCreate.RenderState.NumToProcess == 0)
 	{
@@ -616,7 +617,7 @@ bool UFastGeoContainer::TickDestroyRenderState_Sync(const UE::FTimeout& Timeout)
 
 	TRACE_CPUPROFILER_EVENT_SCOPE(UFastGeoContainer::TickDestroyRenderState_Sync);
 	check(IsInGameThread());
-	check(!FastGeo::Private::GFastGeoAllowAsyncRenderWork);
+	check(!FFastGeoStreamingModule::IsAsyncRenderWorkAllowed());
 
 	FFastGeoComponentClusterPendingState::FRenderState& State = PendingDestroy.RenderState;
 	if (State.ComponentsToProcess.Num() > 0 && State.NumToProcess == 0)
@@ -659,7 +660,7 @@ void UFastGeoContainer::StartCreatePhysicsStateWork()
 {
 	if (const bool bHasPhysicsWork = !CollisionComponents.IsEmpty())
 	{
-		UFastGeoWorldSubsystem* WorldSubsystem = GetWorld()->GetSubsystemChecked<UFastGeoWorldSubsystem>();
+		UFastGeoWorldSubsystem* WorldSubsystem = GetWorldSubsystem();
 		WorldSubsystem->PushAsyncCreatePhysicsStateJobs(this);
 	}
 }
@@ -668,14 +669,14 @@ void UFastGeoContainer::StartDestroyPhysicsStateWork()
 {
 	if (const bool bHasPhysicsWork = !CollisionComponents.IsEmpty())
 	{
-		UFastGeoWorldSubsystem* WorldSubsystem = GetWorld()->GetSubsystemChecked<UFastGeoWorldSubsystem>();
+		UFastGeoWorldSubsystem* WorldSubsystem = GetWorldSubsystem();
 		WorldSubsystem->PushAsyncDestroyPhysicsStateJobs(this);
 	}
 }
 
 void UFastGeoContainer::TickPhysicsWork(bool bWaitForCompletion)
 {
-	if (FPhysScene* PhysScene = GetWorld()->GetPhysicsScene())
+	if (FPhysScene* PhysScene = GetPhysicsScene())
 	{
 		PhysScene->ProcessAsyncPhysicsStateJobs(bWaitForCompletion);
 	}
@@ -684,13 +685,12 @@ void UFastGeoContainer::TickPhysicsWork(bool bWaitForCompletion)
 void UFastGeoContainer::StartPrecachePSOWork()
 {
 #if !WITH_EDITOR && UE_WITH_PSO_PRECACHING
-	if (FastGeo::Private::GFastGeoAllowAsyncRenderWork && !bPrecachedPSOs)
+	if (FFastGeoStreamingModule::IsAsyncRenderWorkAllowed() && !bProcessedPSOPrecaching)
 	{
-		bPrecachedPSOs = true;
-
+		bProcessedPSOPrecaching = true;
 		if (IsComponentPSOPrecachingEnabled())
 		{
-			UFastGeoWorldSubsystem* WorldSubsystem = GetWorld()->GetSubsystemChecked<UFastGeoWorldSubsystem>();
+			UFastGeoWorldSubsystem* WorldSubsystem = GetWorldSubsystem();
 			WorldSubsystem->PushAsyncPrecachePSOsJob(this);
 		}
 	}
@@ -699,26 +699,26 @@ void UFastGeoContainer::StartPrecachePSOWork()
 
 void UFastGeoContainer::StartCreateRenderStateWork()
 {
-	if (FastGeo::Private::GFastGeoAllowAsyncRenderWork && !IsCreateRenderStateCompleted())
+	if (FFastGeoStreamingModule::IsAsyncRenderWorkAllowed() && !IsCreateRenderStateCompleted())
 	{
 #if !WITH_EDITOR && UE_WITH_PSO_PRECACHING
 		// PSO precaching should have been started by UFastGeoWorldSubsystem::OnLevelStreamingStateChanged
-		if (!ensure(bPrecachedPSOs))
+		if (!ensure(bProcessedPSOPrecaching))
 		{
 			StartPrecachePSOWork();
-			check(bPrecachedPSOs);
+			check(bProcessedPSOPrecaching);
 		}
 #endif
-		UFastGeoWorldSubsystem* WorldSubsystem = GetWorld()->GetSubsystemChecked<UFastGeoWorldSubsystem>();
+		UFastGeoWorldSubsystem* WorldSubsystem = GetWorldSubsystem();
 		WorldSubsystem->PushAsyncCreateRenderStateJob(this);
 	}
 }
 
 void UFastGeoContainer::StartDestroyRenderStateWork()
 {
-	if (FastGeo::Private::GFastGeoAllowAsyncRenderWork && !IsDestroyRenderStateCompleted())
+	if (FFastGeoStreamingModule::IsAsyncRenderWorkAllowed() && !IsDestroyRenderStateCompleted())
 	{
-		UFastGeoWorldSubsystem* WorldSubsystem = GetWorld()->GetSubsystemChecked<UFastGeoWorldSubsystem>();
+		UFastGeoWorldSubsystem* WorldSubsystem = GetWorldSubsystem();
 		WorldSubsystem->PushAsyncDestroyRenderStateJob(this);
 	}
 }
@@ -726,7 +726,7 @@ void UFastGeoContainer::StartDestroyRenderStateWork()
 bool UFastGeoContainer::IsPrecachePSOCompleted() const
 {
 #if !WITH_EDITOR && UE_WITH_PSO_PRECACHING
-	return bPrecachedPSOs && PendingCreate.PrecachePSOState.IsCompleted();
+	return bProcessedPSOPrecaching && PendingCreate.PrecachePSOState.IsCompleted();
 #else
 	return true;
 #endif
@@ -744,9 +744,9 @@ bool UFastGeoContainer::IsDestroyRenderStateCompleted() const
 
 void UFastGeoContainer::TickRenderWork(bool bWaitForCompletion)
 {
-	if (FastGeo::Private::GFastGeoAllowAsyncRenderWork)
+	if (FFastGeoStreamingModule::IsAsyncRenderWorkAllowed())
 	{
-		UFastGeoWorldSubsystem* WorldSubsystem = GetWorld()->GetSubsystemChecked<UFastGeoWorldSubsystem>();
+		UFastGeoWorldSubsystem* WorldSubsystem = GetWorldSubsystem();
 		WorldSubsystem->ProcessAsyncRenderStateJobs(bWaitForCompletion);
 	}
 	else
@@ -757,7 +757,7 @@ void UFastGeoContainer::TickRenderWork(bool bWaitForCompletion)
 
 void UFastGeoContainer::TickRenderWork_Sync(bool bWaitForCompletion)
 {
-	if (FastGeo::Private::GFastGeoAllowAsyncRenderWork)
+	if (FFastGeoStreamingModule::IsAsyncRenderWorkAllowed())
 	{
 		return;
 	}
@@ -874,10 +874,64 @@ ULevel* UFastGeoContainer::GetLevel() const
 
 UWorld* UFastGeoContainer::GetWorld() const
 {
+	if (AsyncContext.IsInitialized())
+	{
+		return AsyncContext.GetWorld();
+	}
+	// Fallback when AsyncContext is not initialized (subsystem/time lookups are not safe from workers)
+	check(IsInGameThread() || IsInAsyncLoadingThread() || IsInParallelGameThread());
 	ULevel* Level = GetLevel();
 	return Level ? Level->GetWorld() : nullptr;
 }
 
+UFastGeoWorldSubsystem* UFastGeoContainer::GetWorldSubsystem() const
+{
+	if (AsyncContext.IsInitialized())
+	{
+		return AsyncContext.GetWorldSubsystem();
+	}
+	if (UWorld* World = GetWorld())
+	{
+		return World->GetSubsystemChecked<UFastGeoWorldSubsystem>();
+	}
+	return nullptr;
+}
+
+FSceneInterface* UFastGeoContainer::GetScene() const
+{
+	if (AsyncContext.IsInitialized())
+	{
+		return AsyncContext.GetScene();
+	}
+	if (UWorld* World = GetWorld())
+	{
+		return World->Scene;
+	}
+	return nullptr;
+}
+
+FPhysScene* UFastGeoContainer::GetPhysicsScene() const
+{
+	if (AsyncContext.IsInitialized())
+	{
+		return AsyncContext.GetPhysicsScene();
+	}
+	if (UWorld* World = GetWorld())
+	{
+		return World->GetPhysicsScene();
+	}
+	return nullptr;
+}
+
+float UFastGeoContainer::GetWorldTimeSeconds() const
+{
+	if (AsyncContext.IsInitialized())
+	{
+		return AsyncContext.GetTimeSeconds();
+	}
+	return GetWorld()->GetTimeSeconds();
+}
+
 void UFastGeoContainer::OnNavigationInitDone(const UNavigationSystemBase& InNavigationSystem)
 {
 	const UWorld* World = GetWorld();
@@ -959,9 +1013,12 @@ void UFastGeoContainer::OnCreateRenderStateBegin_GameThread()
 	check(PendingCreate.RenderState.NumToProcess == 0);
 	check(PendingCreate.RenderState.TotalNumProcessed != PendingCreate.RenderState.ComponentsToProcess.Num());
 
-	PendingCreate.bIsInBlockingWait = GetWorld()->GetSubsystemChecked<UFastGeoWorldSubsystem>()->IsWaitingForCompletion();
-
+	PendingCreate.bIsInBlockingWait = GetWorldSubsystem()->IsWaitingForCompletion();
 	PendingCreate.RenderState.NumToProcess = PendingCreate.RenderState.ComponentsToProcess.Num() - PendingCreate.RenderState.TotalNumProcessed;
+
+	// Refresh cached time (see FFastGeoDecalComponent::CreateRenderState)
+	AsyncContext.UpdateTimeSeconds(GetWorld()->GetTimeSeconds());
+
 }
 
 void UFastGeoContainer::OnDestroyRenderStateBegin_GameThread()
@@ -969,7 +1026,7 @@ void UFastGeoContainer::OnDestroyRenderStateBegin_GameThread()
 	check(PendingDestroy.RenderState.NumToProcess == 0);
 	check(PendingDestroy.RenderState.TotalNumProcessed != PendingDestroy.RenderState.ComponentsToProcess.Num());
 
-	PendingDestroy.bIsInBlockingWait = GetWorld()->GetSubsystemChecked<UFastGeoWorldSubsystem>()->IsWaitingForCompletion();
+	PendingDestroy.bIsInBlockingWait = GetWorldSubsystem()->IsWaitingForCompletion();
 
 	// TODO_FASTGEO: Evaluate if we need throttling of the async destruction task
 	PendingDestroy.RenderState.NumToProcess = PendingDestroy.RenderState.ComponentsToProcess.Num();
@@ -977,7 +1034,7 @@ void UFastGeoContainer::OnDestroyRenderStateBegin_GameThread()
 
 void UFastGeoContainer::OnCreateRenderState_Concurrent()
 {
-	check(FastGeo::Private::GFastGeoAllowAsyncRenderWork);
+	check(FFastGeoStreamingModule::IsAsyncRenderWorkAllowed());
 
 	static constexpr int32 MinNumElementsToProcessPerThread = 8;
 	const int32 NumComponentsToProcess = PendingCreate.RenderState.NumToProcess;
@@ -989,7 +1046,7 @@ void UFastGeoContainer::OnCreateRenderState_Concurrent()
 	int32 AvailableComponentsBudget;
 	int32 TimeEpoch;
 
-	UFastGeoWorldSubsystem* WorldSubsystem = GetWorld()->GetSubsystemChecked<UFastGeoWorldSubsystem>();
+	UFastGeoWorldSubsystem* WorldSubsystem = GetWorldSubsystem();
 	WorldSubsystem->RequestAsyncRenderStateTasksBudget_Concurrent(AvailableTimeBudgetMS, AvailableComponentsBudget, TimeEpoch);
 	
 	const int32 ComponentsBudget = FMath::Min(NumComponentsToProcess, AvailableComponentsBudget);
@@ -1056,7 +1113,7 @@ void UFastGeoContainer::OnDestroyRenderState_Concurrent()
 
 void UFastGeoContainer::OnCreateRenderStateEnd_GameThread()
 {
-	if (FastGeo::Private::GFastGeoAllowAsyncRenderWork)
+	if (FFastGeoStreamingModule::IsAsyncRenderWorkAllowed())
 	{
 		check(PendingCreate.RenderState.NumToProcess != 0);
 		check(PendingCreate.RenderState.TotalNumProcessed != PendingCreate.RenderState.ComponentsToProcess.Num());
@@ -1069,7 +1126,7 @@ void UFastGeoContainer::OnCreateRenderStateEnd_GameThread()
 
 	if (!IsCreateRenderStateCompleted())
 	{
-		check(FastGeo::Private::GFastGeoAllowAsyncRenderWork);
+		check(FFastGeoStreamingModule::IsAsyncRenderWorkAllowed());
 		StartCreateRenderStateWork();
 	}
 	else
@@ -1084,7 +1141,7 @@ void UFastGeoContainer::OnCreateRenderStateEnd_GameThread()
 
 void UFastGeoContainer::OnDestroyRenderStateEnd_GameThread()
 {
-	if (FastGeo::Private::GFastGeoAllowAsyncRenderWork)
+	if (FFastGeoStreamingModule::IsAsyncRenderWorkAllowed())
 	{
 		check(PendingDestroy.RenderState.NumToProcess != 0);
 		check(PendingDestroy.RenderState.TotalNumProcessed != PendingDestroy.RenderState.ComponentsToProcess.Num());
@@ -1097,7 +1154,7 @@ void UFastGeoContainer::OnDestroyRenderStateEnd_GameThread()
 
 	if (!IsDestroyRenderStateCompleted())
 	{
-		check(FastGeo::Private::GFastGeoAllowAsyncRenderWork);
+		check(FFastGeoStreamingModule::IsAsyncRenderWorkAllowed());
 		StartDestroyRenderStateWork();
 	}
 	else
@@ -1158,6 +1215,52 @@ void UFastGeoContainer::PreSave(FObjectPreSaveContext ObjectSaveContext)
 	{
 		ComponentCluster.PreSave(ObjectSaveContext);
 	});
+
+	// Collect base materials and Nanite overrides into separate sets.
+	// Base materials go to hard refs. 
+	// Nanite overrides go to soft refs (may be stripped on non-Nanite platforms via FMaterialOverrideNanite::Serialize).
+	if (ObjectSaveContext.IsCooking())
+	{
+		TSet<UMaterialInterface*> BaseSet;
+		TSet<UMaterialInterface*> NaniteOverrideSet;
+
+		ForEachComponentCluster([&BaseSet, &NaniteOverrideSet](const FFastGeoComponentCluster& Cluster)
+		{
+			Cluster.ForEachComponent<FFastGeoComponent>([&BaseSet, &NaniteOverrideSet](const FFastGeoComponent& Comp)
+			{
+				Comp.ForEachMaterial([&BaseSet, &NaniteOverrideSet](UMaterialInterface* Mat, bool bIsNaniteOverride)
+				{
+					if (bIsNaniteOverride)
+					{
+						NaniteOverrideSet.Add(Mat);
+					}
+					else
+					{
+						BaseSet.Add(Mat);
+					}
+				});
+			});
+		});
+
+		// If a Nanite override is also used directly as a base material, keep it in hard refs only
+		NaniteOverrideSet = NaniteOverrideSet.Difference(BaseSet);
+
+		MaterialAssetHardRefs.Reset(BaseSet.Num());
+		for (UMaterialInterface* Mat : BaseSet)
+		{
+			MaterialAssetHardRefs.Add(Mat);
+		}
+
+		MaterialAssetNaniteOverrideSoftRefs.Reset(NaniteOverrideSet.Num());
+		for (UMaterialInterface* Mat : NaniteOverrideSet)
+		{
+			MaterialAssetNaniteOverrideSoftRefs.Emplace(Mat);
+		}
+
+		// Clear transient cache (will be rebuilt at load time)
+		CachedMaterialAssets.Reset();
+		bMaterialAssetsCached = false;
+	}
 }
 #endif
 
@@ -1255,6 +1358,11 @@ void UFastGeoContainer::Serialize(FArchive& Ar)
 	{
 		// Once loaded, initialize component clusters & components dynamic properties
 		InitializeDynamicProperties(true);
+
+		// Resolve pre-computed material hard/soft refs into one single cached array.
+		// All material packages are loaded (hard-referenced via Assets).
+		// Stripped materials like Nanite overrides on non-Nanite platform will be skipped.
+		CacheMaterialAssets();
 	}
 #endif
 }
@@ -1265,11 +1373,23 @@ void UFastGeoContainer::SerializeComponentClusters(FArchive& Ar)
 	Ar << HLODs;
 }
 
+void UFastGeoContainer::OnPrecachePSOsBegin_GameThread()
+{
+#if !WITH_EDITOR && UE_WITH_PSO_PRECACHING
+	TRACE_CPUPROFILER_EVENT_SCOPE(UFastGeoContainer::OnPrecachePSOsBegin_GameThread);
+	check(FFastGeoStreamingModule::IsAsyncRenderWorkAllowed());
+	check(IsComponentPSOPrecachingEnabled());
+	
+	// PSO precaching can run before StartRegisterTransition, so initialize all prerequisites now.
+	InitializePrerequisites();
+#endif
+}
+
 void UFastGeoContainer::OnPrecachePSOs_Concurrent()
 {
 #if !WITH_EDITOR && UE_WITH_PSO_PRECACHING
 	TRACE_CPUPROFILER_EVENT_SCOPE(UFastGeoContainer::OnPrecachePSOs_Concurrent);
-
+	check(FFastGeoStreamingModule::IsAsyncRenderWorkAllowed());
 	check(IsComponentPSOPrecachingEnabled());
 
 	ForEachComponentCluster([](FFastGeoComponentCluster& ComponentCluster)
@@ -1289,6 +1409,172 @@ void UFastGeoContainer::PrecachePSOs()
 	StartPrecachePSOWork();
 }
 
+void UFastGeoContainer::ForEachMaterialAsset(TFunctionRef<void(UMaterialInterface*)> Func)
+{
+	// Cooked builds: already resolved from soft refs in Serialize (ALT).
+	// PIE / Transient FastGeoContainer: build from component walk on first call (runs once on GT).
+	if (!bMaterialAssetsCached)
+	{
+		BuildMaterialAssets();
+	}
+
+	for (UMaterialInterface* Mat : CachedMaterialAssets)
+	{
+		Func(Mat);
+	}
+}
+
+void UFastGeoContainer::BuildMaterialAssets()
+{
+	TRACE_CPUPROFILER_EVENT_SCOPE(UFastGeoContainer::BuildMaterialAssets);
+
+	TSet<UMaterialInterface*> UniqueSet;
+
+	ForEachComponentCluster([&UniqueSet](const FFastGeoComponentCluster& Cluster)
+	{
+		Cluster.ForEachComponent<FFastGeoComponent>([&UniqueSet](const FFastGeoComponent& Comp)
+		{
+			Comp.ForEachMaterial([&UniqueSet](UMaterialInterface* Mat, bool /*bIsNaniteOverride*/)
+			{
+				UniqueSet.Add(Mat);
+			});
+		});
+	});
+
+	CachedMaterialAssets = UniqueSet.Array();
+	bMaterialAssetsCached = true;
+}
+
+#if !WITH_EDITOR
+void UFastGeoContainer::CacheMaterialAssets()
+{
+	TRACE_CPUPROFILER_EVENT_SCOPE(UFastGeoContainer::CacheMaterialAssets);
+
+	CachedMaterialAssets = MaterialAssetHardRefs;
+	CachedMaterialAssets.Reserve(CachedMaterialAssets.Num() + MaterialAssetNaniteOverrideSoftRefs.Num());
+
+	// Resolve Nanite overrides, stripped materials (e.g. Nanite override on non-Nanite platform) resolve to null and are skipped
+	for (const TSoftObjectPtr<UMaterialInterface>& SoftRef : MaterialAssetNaniteOverrideSoftRefs)
+	{
+		if (UMaterialInterface* Mat = SoftRef.Get())
+		{
+			CachedMaterialAssets.Add(Mat);
+		}
+	}
+
+	bMaterialAssetsCached = true;
+}
+#endif
+
+/**
+ * Cache world/scene/physics/subsystem pointers for worker-thread access.
+ * Workers read cached pointers instead of performing lookups not safe from
+ * generic workers (e.g., subsystem collection access, world time reads).
+ */
+void FFastGeoAsyncContext::Initialize(UFastGeoContainer& Container)
+{
+	if (bInitialized)
+	{
+		return;
+	}
+
+	check(IsInGameThread());
+
+	World = Container.GetWorld();
+	check(World);
+	Scene = World->Scene;
+	check(Scene);
+	PhysicsScene = World->GetPhysicsScene();
+	check(PhysicsScene);
+	WorldSubsystem = World->GetSubsystemChecked<UFastGeoWorldSubsystem>();
+	TimeSeconds = World->GetTimeSeconds();
+	bInitialized = true;
+}
+
+/**
+ * Pre-call ConditionalPostLoad on all materials referenced by this container.
+ * This is necessary because we don't want any of these function to call
+ * ConditionalPostLoad on materials from worker threads:
+ *     - UMaterialInstance::PrecachePSOs
+ *     - FStaticMeshComponentHelper::GetMaterial
+ *     - FStaticMeshComponentHelper::GetUsedRayTracingOnlyMaterials
+ * 
+ * Note that once the preemptive ConditionalPostLoad is removed from these functions,
+ * this prerequisite step can be removed.
+ */
+void FFastGeoMaterialPostLoadInit::Initialize(UFastGeoContainer& Container)
+{
+	if (bInitialized)
+	{
+		return;
+	}
+
+	TRACE_CPUPROFILER_EVENT_SCOPE(FFastGeoMaterialPostLoadInit::Initialize);
+	check(IsInGameThread());
+
+	Container.ForEachMaterialAsset([](UMaterialInterface* Mat)
+	{
+		if (Mat->HasAnyFlags(RF_NeedPostLoad))
+		{
+			Mat->ConditionalPostLoad();
+		}
+	});
+
+	bInitialized = true;
+}
+
+/**
+ * Warm all materials relevance cache because so that FastGeo 
+ * workers call to functions below will the thread safe:
+ *     - FMeshComponentHelper::GetMaterialRelevance
+ *     - FNaniteResourcesHelper::AuditMaterials
+ *     - FStaticMeshComponentHelper::CollectPSOPrecacheDataImpl
+ */
+void FFastGeoMaterialRelevanceInit::Initialize(UFastGeoContainer& Container)
+{
+#if !WITH_EDITOR
+	if (bInitialized || !FFastGeoStreamingModule::IsAsyncRenderWorkAllowed())
+	{
+		return;
+	}
+
+	TRACE_CPUPROFILER_EVENT_SCOPE(FFastGeoMaterialRelevanceInit::Initialize);
+	check(IsInGameThread());
+
+	const EShaderPlatform Platform = Container.GetWorld()->Scene->GetShaderPlatform();
+
+	Container.ForEachMaterialAsset([Platform](UMaterialInterface* Mat)
+	{
+		Mat->WarmRelevanceCache(Platform);
+	});
+
+	// Also pre-warm default material on the GameThread (UMaterial::GetDefaultMaterial lazy init has check(IsInGameThread())).
+	if (UMaterialInterface* DefaultMaterial = UMaterial::GetDefaultMaterial(MD_Surface))
+	{
+		DefaultMaterial->WarmRelevanceCache(Platform);
+	}
+
+#if !WITH_EDITOR && UE_WITH_PSO_PRECACHING
+	// Also pre-warm PSO fallback material
+	if (UMaterialInterface* PSOFallbackMaterial = UPSOPrecacheSettingsManager::GetFallbackMaterial())
+	{
+		PSOFallbackMaterial->WarmRelevanceCache(Platform);
+	}
+#endif
+
+	bInitialized = true;
+#endif
+}
+
+void UFastGeoContainer::InitializePrerequisites()
+{
+	check(IsInGameThread());
+
+	AsyncContext.Initialize(*this);
+	MaterialPostLoadInit.Initialize(*this);
+	MaterialRelevanceInit.Initialize(*this);
+}
+
 void UFastGeoContainer::InitializeDynamicProperties(bool bInitForPlay)
 {
 	ForEachComponentCluster([this](FFastGeoComponentCluster& ComponentCluster)
@@ -1337,7 +1623,7 @@ void UFastGeoContainer::OnCreatePhysicsStateBegin_GameThread()
 
 		UWorld* World = GetWorld();
 		check(World);
-		FPhysScene* PhysScene = World->GetPhysicsScene();
+		FPhysScene* PhysScene = GetPhysicsScene();
 		check(PhysScene);
 		PhysScene->PushAsyncCreatePhysicsStateJob(this);
 	}
@@ -1356,7 +1642,7 @@ void UFastGeoContainer::OnDestroyPhysicsStateBegin_GameThread()
 
 		UWorld* World = GetWorld();
 		check(World);
-		FPhysScene* PhysScene = World->GetPhysicsScene();
+		FPhysScene* PhysScene = GetPhysicsScene();
 		check(PhysScene);
 		verify(PhysScene->PushAsyncDestroyPhysicsStateJob(this));
 	}
diff --git a/Engine/Plugins/Experimental/FastGeoStreaming/Source/FastGeoStreaming/Private/FastGeoDecalComponent.cpp b/Engine/Plugins/Experimental/FastGeoStreaming/Source/FastGeoStreaming/Private/FastGeoDecalComponent.cpp
--- a/Engine/Plugins/Experimental/FastGeoStreaming/Source/FastGeoStreaming/Private/FastGeoDecalComponent.cpp
+++ b/Engine/Plugins/Experimental/FastGeoStreaming/Source/FastGeoStreaming/Private/FastGeoDecalComponent.cpp
@@ -45,6 +45,14 @@ void FFastGeoDecalComponent::InitializeFromComponent(UActorComponent* Component)
 }
 #endif
 
+void FFastGeoDecalComponent::ForEachMaterial(TFunctionRef<void(UMaterialInterface*, bool bIsNaniteOverride)> Func) const
+{
+	if (SceneProxyDesc.DecalMaterial)
+	{
+		Func(SceneProxyDesc.DecalMaterial, false);
+	}
+}
+
 void FFastGeoDecalComponent::Serialize(FArchive& Ar)
 {
 	Super::Serialize(Ar);
@@ -117,7 +125,7 @@ void FFastGeoDecalComponent::CreateRenderState(FRegisterComponentContext* Contex
 	SceneProxy = CreateSceneProxy(Error);
 	if (SceneProxy)
 	{
-		const float SpawnTime = GetWorld()->GetTimeSeconds();
+		const float SpawnTime = GetWorldTimeSeconds();
 		TArray<FDeferredDecalUpdateParams> DecalsToAdd;
 		DecalsToAdd.Add(
 		{
diff --git a/Engine/Plugins/Experimental/FastGeoStreaming/Source/FastGeoStreaming/Private/FastGeoInstancedStaticMeshComponent.cpp b/Engine/Plugins/Experimental/FastGeoStreaming/Source/FastGeoStreaming/Private/FastGeoInstancedStaticMeshComponent.cpp
--- a/Engine/Plugins/Experimental/FastGeoStreaming/Source/FastGeoStreaming/Private/FastGeoInstancedStaticMeshComponent.cpp
+++ b/Engine/Plugins/Experimental/FastGeoStreaming/Source/FastGeoStreaming/Private/FastGeoInstancedStaticMeshComponent.cpp
@@ -271,7 +271,7 @@ void FFastGeoInstancedStaticMeshComponent::OnAsyncCreatePhysicsState()
 
 	check(InstanceBodies.Num() == 0);
 
-	FPhysScene* PhysScene = GetWorld()->GetPhysicsScene();
+	FPhysScene* PhysScene = GetPhysicsScene();
 	if (!PhysScene)
 	{
 		return;
@@ -397,7 +397,7 @@ void FFastGeoInstancedStaticMeshComponent::CreateAllInstanceBodies()
 
 	if (UBodySetup* BodySetup = GetBodySetup())
 	{
-		FPhysScene* PhysScene = GetWorld()->GetPhysicsScene();
+		FPhysScene* PhysScene = GetPhysicsScene();
 
 		if (!BodyInstance.GetOverrideWalkableSlopeOnInstance())
 		{
@@ -447,7 +447,7 @@ void FFastGeoInstancedStaticMeshComponent::CreateAllInstanceBodies()
 			BodyInstanceOwner.Initialize(this);
 
 			// Initialize body instances
-			FBodyInstance::InitStaticBodies(MoveTemp(InstanceBodiesSanitized), MoveTemp(Transforms), BodySetup, SurrogateComponent, GetWorld()->GetPhysicsScene(), &BodyInstanceOwner);
+			FBodyInstance::InitStaticBodies(MoveTemp(InstanceBodiesSanitized), MoveTemp(Transforms), BodySetup, SurrogateComponent, GetPhysicsScene(), &BodyInstanceOwner);
 
 			// Assign BodyInstanceOwner
 			TArray<Chaos::FPhysicsObject*> PhysicsObjects = GetAllPhysicsObjects(InstanceBodies);
diff --git a/Engine/Plugins/Experimental/FastGeoStreaming/Source/FastGeoStreaming/Private/FastGeoLocalLightComponent.cpp b/Engine/Plugins/Experimental/FastGeoStreaming/Source/FastGeoStreaming/Private/FastGeoLocalLightComponent.cpp
--- a/Engine/Plugins/Experimental/FastGeoStreaming/Source/FastGeoStreaming/Private/FastGeoLocalLightComponent.cpp
+++ b/Engine/Plugins/Experimental/FastGeoStreaming/Source/FastGeoStreaming/Private/FastGeoLocalLightComponent.cpp
@@ -42,6 +42,15 @@ void FFastGeoLocalLightComponent::InitializeFromComponent(UActorComponent* Compo
 }
 #endif
 
+void FFastGeoLocalLightComponent::ForEachMaterial(TFunctionRef<void(UMaterialInterface*, bool bIsNaniteOverride)> Func) const
+{
+	const FLocalLightSceneProxyDesc& SceneProxyDesc = GetLocalLightSceneProxyDesc();
+	if (SceneProxyDesc.LightFunctionMaterial)
+	{
+		Func(SceneProxyDesc.LightFunctionMaterial, false);
+	}
+}
+
 void FFastGeoLocalLightComponent::Serialize(FArchive& Ar)
 {
 	Super::Serialize(Ar);
diff --git a/Engine/Plugins/Experimental/FastGeoStreaming/Source/FastGeoStreaming/Private/FastGeoPrimitiveComponent.cpp b/Engine/Plugins/Experimental/FastGeoStreaming/Source/FastGeoStreaming/Private/FastGeoPrimitiveComponent.cpp
--- a/Engine/Plugins/Experimental/FastGeoStreaming/Source/FastGeoStreaming/Private/FastGeoPrimitiveComponent.cpp
+++ b/Engine/Plugins/Experimental/FastGeoStreaming/Source/FastGeoStreaming/Private/FastGeoPrimitiveComponent.cpp
@@ -431,7 +431,7 @@ void FFastGeoPrimitiveComponent::OnAsyncCreatePhysicsState()
 			check(!SurrogateComponent || FFastGeoSurrogateBodyInstanceIndex::IsEncoded(BodyInstance.InstanceBodyIndex));
 
 			// Initialize the body instance
-			BodyInstance.InitBody(BodySetup, BodyTransform, SurrogateComponent, GetWorld()->GetPhysicsScene(), FInitBodySpawnParams(IsStaticPhysics(), /*bPhysicsTypeDeterminesSimulation*/false), &BodyInstanceOwner);
+			BodyInstance.InitBody(BodySetup, BodyTransform, SurrogateComponent, GetPhysicsScene(), FInitBodySpawnParams(IsStaticPhysics(), /*bPhysicsTypeDeterminesSimulation*/false), &BodyInstanceOwner);
 
 			// Assign BodyInstanceOwner
 			if (Chaos::FSingleParticlePhysicsProxy* ProxyHandle = BodyInstance.GetPhysicsActor())
diff --git a/Engine/Plugins/Experimental/FastGeoStreaming/Source/FastGeoStreaming/Private/FastGeoSkinnedMeshComponent.cpp b/Engine/Plugins/Experimental/FastGeoStreaming/Source/FastGeoStreaming/Private/FastGeoSkinnedMeshComponent.cpp
--- a/Engine/Plugins/Experimental/FastGeoStreaming/Source/FastGeoStreaming/Private/FastGeoSkinnedMeshComponent.cpp
+++ b/Engine/Plugins/Experimental/FastGeoStreaming/Source/FastGeoStreaming/Private/FastGeoSkinnedMeshComponent.cpp
@@ -204,6 +204,37 @@ void FFastGeoSkinnedMeshComponentBase::DestroyRenderState(FFastGeoDestroyRenderS
 	Super::DestroyRenderState(Context);
 }
 
+void FFastGeoSkinnedMeshComponentBase::ForEachMaterial(TFunctionRef<void(UMaterialInterface*, bool bIsNaniteOverride)> Func) const
+{
+	auto AddMaterial = [&Func](UMaterialInterface* Mat)
+	{
+		Func(Mat, false);
+		if (UMaterialInterface* NaniteOverride = Mat->GetNaniteOverride())
+		{
+			Func(NaniteOverride, true);
+		}
+	};
+
+	for (const TObjectPtr<UMaterialInterface>& Mat : OverrideMaterials)
+	{
+		if (Mat)
+		{
+			AddMaterial(Mat);
+		}
+	}
+
+	if (USkinnedAsset* Asset = GetSkinnedAsset())
+	{
+		for (const FSkeletalMaterial& SlotMat : Asset->GetMaterials())
+		{
+			if (SlotMat.MaterialInterface)
+			{
+				AddMaterial(SlotMat.MaterialInterface);
+			}
+		}
+	}
+}
+
 UMaterialInterface* FFastGeoSkinnedMeshComponentBase::GetMaterial(int32 MaterialIndex) const
 {
 	return FSkinnedMeshComponentHelper::GetMaterial(*this, MaterialIndex);
diff --git a/Engine/Plugins/Experimental/FastGeoStreaming/Source/FastGeoStreaming/Private/FastGeoStaticMeshComponent.cpp b/Engine/Plugins/Experimental/FastGeoStreaming/Source/FastGeoStreaming/Private/FastGeoStaticMeshComponent.cpp
--- a/Engine/Plugins/Experimental/FastGeoStreaming/Source/FastGeoStreaming/Private/FastGeoStaticMeshComponent.cpp
+++ b/Engine/Plugins/Experimental/FastGeoStreaming/Source/FastGeoStreaming/Private/FastGeoStaticMeshComponent.cpp
@@ -192,6 +192,37 @@ UMaterialInterface* FFastGeoStaticMeshComponentBase::GetNaniteAuditMaterial(int3
 	return GetMaterial(MaterialIndex, true);
 }
 
+void FFastGeoStaticMeshComponentBase::ForEachMaterial(TFunctionRef<void(UMaterialInterface*, bool bIsNaniteOverride)> Func) const
+{
+	auto AddMaterial = [&Func](UMaterialInterface* Mat)
+	{
+		Func(Mat, false);
+		if (UMaterialInterface* NaniteOverride = Mat->GetNaniteOverride())
+		{
+			Func(NaniteOverride, true);
+		}
+	};
+
+	for (const TObjectPtr<UMaterialInterface>& Mat : OverrideMaterials)
+	{
+		if (Mat)
+		{
+			AddMaterial(Mat);
+		}
+	}
+
+	if (UStaticMesh* Mesh = GetStaticMesh())
+	{
+		for (const FStaticMaterial& SlotMat : Mesh->GetStaticMaterials())
+		{
+			if (SlotMat.MaterialInterface)
+			{
+				AddMaterial(SlotMat.MaterialInterface);
+			}
+		}
+	}
+}
+
 bool FFastGeoStaticMeshComponentBase::UseNaniteOverrideMaterials(bool bDoingMaterialAudit) const
 {
 	return Nanite::FNaniteResourcesHelper::UseNaniteOverrideMaterials(*this, bDoingMaterialAudit);
diff --git a/Engine/Plugins/Experimental/FastGeoStreaming/Source/FastGeoStreaming/Private/FastGeoStreamingModule.cpp b/Engine/Plugins/Experimental/FastGeoStreaming/Source/FastGeoStreaming/Private/FastGeoStreamingModule.cpp
--- a/Engine/Plugins/Experimental/FastGeoStreaming/Source/FastGeoStreaming/Private/FastGeoStreamingModule.cpp
+++ b/Engine/Plugins/Experimental/FastGeoStreaming/Source/FastGeoStreaming/Private/FastGeoStreamingModule.cpp
@@ -18,8 +18,33 @@ FAutoConsoleVariableRef FFastGeoStreamingModule::CVarIsEnabled(
 	TEXT("Set to false to disable FastGeo (used in PIE and at cook time)."),
 	ECVF_Default | ECVF_ReadOnly);
 
+#if !WITH_EDITOR
+bool FFastGeoStreamingModule::bAllowAsyncRenderWork = true;
+FAutoConsoleVariableRef FFastGeoStreamingModule::CVarAllowAsyncRenderWork(
+	TEXT("FastGeo.AllowAsyncRenderWork"),
+	FFastGeoStreamingModule::bAllowAsyncRenderWork,
+	TEXT("Allow FastGeo to perform render state creation/destruction and PSO precaching asynchronously. Cooked builds only."),
+	ECVF_ReadOnly);
+#endif
+
+bool FFastGeoStreamingModule::IsAsyncRenderWorkAllowed()
+{
+#if WITH_EDITOR
+	return false;
+#else
+	return bAllowAsyncRenderWork;
+#endif
+}
+
 void FFastGeoStreamingModule::StartupModule()
 {
+#if !WITH_EDITOR
+	if (IConsoleVariable* SSAMCVar = IConsoleManager::Get().FindConsoleVariable(TEXT("s.StreamableAssets.UseSimpleStreamableAssetManager")))
+	{
+		ensureMsgf(SSAMCVar->GetBool(), TEXT("FSimpleStreamableAssetManager must be enabled for FastGeo. Set s.StreamableAssets.UseSimpleStreamableAssetManager=1"));
+	}
+#endif
+
 	IPrimitiveComponent::AddImplementer(TComponentInterfaceImplementation<IPrimitiveComponent>(UFastGeoContainer::StaticClass(),  [](UObject* Obj)
 	{
 		return CastChecked<UFastGeoContainer>(Obj)->GetPrimitiveComponents();
diff --git a/Engine/Plugins/Experimental/FastGeoStreaming/Source/FastGeoStreaming/Public/FastGeoStreamingModule.h b/Engine/Plugins/Experimental/FastGeoStreaming/Source/FastGeoStreaming/Public/FastGeoStreamingModule.h
--- a/Engine/Plugins/Experimental/FastGeoStreaming/Source/FastGeoStreaming/Public/FastGeoStreamingModule.h
+++ b/Engine/Plugins/Experimental/FastGeoStreaming/Source/FastGeoStreaming/Public/FastGeoStreamingModule.h
@@ -18,6 +18,7 @@ public:
 	virtual void ShutdownModule() override;
 
 	FASTGEOSTREAMING_API static bool IsFastGeoEnabled();
+	FASTGEOSTREAMING_API static bool IsAsyncRenderWorkAllowed();
 
 private:
 	void OnWorldPreSendAllEndOfFrameUpdates(UWorld* InWorld);
@@ -27,4 +28,9 @@ private:
 
 	static bool bIsFastGeoEnabled;
 	static FAutoConsoleVariableRef CVarIsEnabled;
+
+#if !WITH_EDITOR
+	static bool bAllowAsyncRenderWork;
+	static FAutoConsoleVariableRef CVarAllowAsyncRenderWork;
+#endif
 };
diff --git a/Engine/Source/Runtime/Engine/Private/Materials/MaterialInterface.cpp b/Engine/Source/Runtime/Engine/Private/Materials/MaterialInterface.cpp
--- a/Engine/Source/Runtime/Engine/Private/Materials/MaterialInterface.cpp
+++ b/Engine/Source/Runtime/Engine/Private/Materials/MaterialInterface.cpp
@@ -578,6 +578,12 @@ bool UMaterialInterface::IsUsingWorldPositionOffset_Concurrent(ERHIFeatureLevel:
 
 bool UMaterialInterface::IsUsingWorldPositionOffset_Concurrent(EShaderPlatform InShaderPlatform) const
 {
+#if !WITH_EDITOR
+	if (const FCachedMaterialRelevance* Cached = CachedMaterialRelevance.TryGet())
+	{
+		return Cached->Relevance.bUsesWorldPositionOffset;
+	}
+#endif
 	if (const UMaterial* Material = GetMaterial_Concurrent())
 	{
 		const FMaterialResource* MaterialResource = GetMaterialResource(InShaderPlatform);
@@ -595,6 +601,12 @@ bool UMaterialInterface::IsUsingWorldPositionOffset_Concurrent(EShaderPlatform I
 
 bool UMaterialInterface::IsUsingMotionVectorWorldOffset_Concurrent(EShaderPlatform InShaderPlatform) const
 {
+#if !WITH_EDITOR
+	if (const FCachedMaterialRelevance* Cached = CachedMaterialRelevance.TryGet())
+	{
+		return Cached->bUsesMotionVectorWorldOffset;
+	}
+#endif
 	if (const UMaterial* Material = GetMaterial_Concurrent())
 	{
 		const FMaterialResource* MaterialResource = GetMaterialResource(InShaderPlatform);
@@ -612,6 +624,12 @@ bool UMaterialInterface::IsUsingMotionVectorWorldOffset_Concurrent(EShaderPlatfo
 
 bool UMaterialInterface::IsUsingTemporalResponsiveness_Concurrent(EShaderPlatform InShaderPlatform) const
 {
+#if !WITH_EDITOR
+	if (const FCachedMaterialRelevance* Cached = CachedMaterialRelevance.TryGet())
+	{
+		return Cached->bUsesTemporalResponsiveness;
+	}
+#endif
 	if (const UMaterial* Material = GetMaterial_Concurrent())
 	{
 		const FMaterialResource* MaterialResource = GetMaterialResource(InShaderPlatform);
@@ -714,7 +732,7 @@ FMaterialRelevance UMaterialInterface::GetRelevance_Internal(const UMaterial* Ma
 	MaterialRelevance.bTranslucencyModulate = bMaterialSeparateModulation;
 	MaterialRelevance.bPostMotionBlurTranslucency = (TranslucencyPass == MTP_AfterMotionBlur);
 	MaterialRelevance.bNormalTranslucency = bIsTranslucent && (TranslucencyPass == MTP_BeforeDOF);
-	MaterialRelevance.bDisableDepthTest = bIsTranslucent && Material->bDisableDepthTest;		
+	MaterialRelevance.bDisableDepthTest = bIsTranslucent && Material->bDisableDepthTest;
 	MaterialRelevance.bUsesSceneColorCopy = bIsTranslucent && MaterialResource->RequiresSceneColorCopy_GameThread();
 	MaterialRelevance.bOutputsTranslucentVelocity = MaterialResource->IsTranslucencyWritingVelocity();
 	MaterialRelevance.bUsesGlobalDistanceField = MaterialResource->UsesGlobalDistanceField_GameThread();
@@ -773,6 +791,12 @@ FMaterialRelevance UMaterialInterface::GetRelevance(ERHIFeatureLevel::Type InFea
 
 FMaterialRelevance UMaterialInterface::GetRelevance(EShaderPlatform InShaderPlatform) const
 {
+#if !WITH_EDITOR
+	if (const FCachedMaterialRelevance* Cached = CachedMaterialRelevance.TryGet())
+	{
+		return Cached->Relevance;
+	}
+#endif
 	// Find the interface's concrete material.
 	const UMaterial* Material = GetMaterial();
 	return GetRelevance_Internal(Material, InShaderPlatform);
@@ -786,11 +810,32 @@ FMaterialRelevance UMaterialInterface::GetRelevance_Concurrent(ERHIFeatureLevel:
 
 FMaterialRelevance UMaterialInterface::GetRelevance_Concurrent(EShaderPlatform InShaderPlatform) const
 {
+#if !WITH_EDITOR
+	if (const FCachedMaterialRelevance* Cached = CachedMaterialRelevance.TryGet())
+	{
+		return Cached->Relevance;
+	}
+#endif
 	// Find the interface's concrete material.
 	const UMaterial* Material = GetMaterial_Concurrent();
 	return GetRelevance_Internal(Material, InShaderPlatform);
 }
 
+#if !WITH_EDITOR
+void UMaterialInterface::WarmRelevanceCache(EShaderPlatform InShaderPlatform)
+{
+	if (!CachedMaterialRelevance.TryGet())
+	{
+		check(IsInGameThread());
+		FCachedMaterialRelevance Cache;
+		Cache.Relevance = GetRelevance_Concurrent(InShaderPlatform);
+		Cache.bUsesMotionVectorWorldOffset = IsUsingMotionVectorWorldOffset_Concurrent(InShaderPlatform);
+		Cache.bUsesTemporalResponsiveness = IsUsingTemporalResponsiveness_Concurrent(InShaderPlatform);
+		CachedMaterialRelevance.Set(Cache);
+	}
+}
+#endif // !WITH_EDITOR
+
 int32 UMaterialInterface::GetWidth() const
 {
 	return ME_PREV_THUMBNAIL_SZ+(ME_STD_BORDER*2);
diff --git a/Engine/Source/Runtime/Engine/Public/Materials/MaterialInterface.h b/Engine/Source/Runtime/Engine/Public/Materials/MaterialInterface.h
--- a/Engine/Source/Runtime/Engine/Public/Materials/MaterialInterface.h
+++ b/Engine/Source/Runtime/Engine/Public/Materials/MaterialInterface.h
@@ -418,6 +418,54 @@ protected:
 	UPROPERTY(Transient)
 	TOptional<FMaterialCachedTexturesSamplingInfo> CachedTexturesSamplingInfo;
 
+#if !WITH_EDITOR
+	/**
+	 * Cached material relevance for async worker-thread access (cooked builds only).
+	 * Populated on GameThread via WarmRelevanceCache(). Generic workers read from cache
+	 * instead of going through GetRelevance_Internal -> _GameThread() methods
+	 * which read GameThreadShaderMap which is not safe from generic workers.
+	 *
+	 * bUsesMotionVectorWorldOffset and bUsesTemporalResponsiveness are not part of
+	 * FMaterialRelevance because it is exactly 64 bits with no room for additional fields.
+	 */
+	struct FCachedMaterialRelevance
+	{
+		FMaterialRelevance Relevance;
+		uint8 bUsesMotionVectorWorldOffset : 1 = 0;
+		uint8 bUsesTemporalResponsiveness : 1 = 0;
+	};
+
+	/**
+	 * Thread-safe write-once wrapper. 
+	 * GameThread writes via Set(), workers read via TryGet(). 
+	 */
+	struct FAtomicCachedMaterialRelevance
+	{
+		void Set(const FCachedMaterialRelevance& InValue)
+		{
+			Value = InValue;
+			bReady.store(true, std::memory_order_release);
+		}
+
+		// Cached relevance is not keyed by shader platform. In cooked builds the shader platform
+		// is fixed at startup and cannot change at runtime, so this is safe.
+		const FCachedMaterialRelevance* TryGet() const
+		{
+			if (bReady.load(std::memory_order_acquire))
+			{
+				return &Value;
+			}
+			return nullptr;
+		}
+
+	private:
+		FCachedMaterialRelevance Value;
+		std::atomic<bool> bReady{false};
+	};
+
+	FAtomicCachedMaterialRelevance CachedMaterialRelevance;
+#endif
+
 private:
 	/** Feature levels to force to compile. */
 	uint32 FeatureLevelsToForceCompile;
@@ -879,6 +927,15 @@ public:
 	ENGINE_API FMaterialRelevance GetRelevance_Concurrent(ERHIFeatureLevel::Type InFeatureLevel) const;
 	ENGINE_API FMaterialRelevance GetRelevance_Concurrent(EShaderPlatform ShaderPlatform) const;
 
+#if !WITH_EDITOR
+	/**
+	 * Pre-warm the cached relevance on GT. No-op if already cached.
+	 * Generic workers (e.g., FastGeo async mode) use the cached value via GetRelevance_Concurrent
+	 * to avoid reading GameThreadShaderMap which is not safe from generic workers.
+	 */
+	void WarmRelevanceCache(EShaderPlatform InShaderPlatform);
+#endif
+
 	/** @return The material's relevance bUsesWorldPositionOffset, from concurrent render thread updates. */
 	UE_DEPRECATED(5.7, "Please use IsUsingWorldPositionOffset_Concurrent with EShaderPlatform argument and not ERHIFeatureLevel::Type")
 	ENGINE_API bool IsUsingWorldPositionOffset_Concurrent(ERHIFeatureLevel::Type InFeatureLevel) const;
