[FastGeo] Change default ISM instance storage mode from Full to Auto
- Auto mode compresses instance transforms from FMatrix (136 bytes) to FTransform3f (52 bytes), falling back to full precision only when compression would cause precision loss.
- Profiled on console with an automated flythrough in CitySample Big_City_LVL (44k ISM components, 1.4M instances, 840 frames):
. Memory savings: 113.6 MB (184 MB -> 70 MB, 62% reduction)
. No measurable frame time impact (within noise)
. 100% of components compressed with no precision loss

#Synced-CL 51624647

#Synced-CL 51624647

diff --git a/Engine/Plugins/Experimental/FastGeoStreaming/Source/FastGeoStreaming/Public/FastGeoWorldPartitionRuntimeCellTransformer.h b/Engine/Plugins/Experimental/FastGeoStreaming/Source/FastGeoStreaming/Public/FastGeoWorldPartitionRuntimeCellTransformer.h
--- a/Engine/Plugins/Experimental/FastGeoStreaming/Source/FastGeoStreaming/Public/FastGeoWorldPartitionRuntimeCellTransformer.h
+++ b/Engine/Plugins/Experimental/FastGeoStreaming/Source/FastGeoStreaming/Public/FastGeoWorldPartitionRuntimeCellTransformer.h
@@ -150,7 +150,7 @@ public:
 
 	/** Determines how instance transforms are stored in memory */
 	UPROPERTY(EditAnywhere, Category = "Settings")
-	EFastGeoInstanceStorageMode InstanceStorageMode = EFastGeoInstanceStorageMode::Full;
+	EFastGeoInstanceStorageMode InstanceStorageMode = EFastGeoInstanceStorageMode::Auto;
 
 	UPROPERTY(EditAnywhere, Category = "Settings", meta = (EditCondition = "InstanceStorageMode == EFastGeoInstanceStorageMode::Auto", EditConditionHides))
 	FFastGeoInstanceCompressionSettings InstanceCompressionSettings;
