[FastGeo] Fixed ForceUpdateVisibility looping indefinitely during blocking loads
When primitives are not yet registered in the scene during a blocking level load, ForceUpdateVisibility retries can loop indefinitely between
the render thread and game thread. After two immediate retries, defer to the next frame tick via FTSTicker to allow pending AddPrimitive
render commands to complete.

#Synced-CL 51624647

diff --git a/Engine/Plugins/Experimental/FastGeoStreaming/Source/FastGeoStreaming/Internal/FastGeoComponentCluster.h b/Engine/Plugins/Experimental/FastGeoStreaming/Source/FastGeoStreaming/Internal/FastGeoComponentCluster.h
--- a/Engine/Plugins/Experimental/FastGeoStreaming/Source/FastGeoStreaming/Internal/FastGeoComponentCluster.h
+++ b/Engine/Plugins/Experimental/FastGeoStreaming/Source/FastGeoStreaming/Internal/FastGeoComponentCluster.h
@@ -179,6 +179,7 @@ protected:
 
 private:
 	void ForceUpdateVisibility(const TArray<FFastGeoPrimitiveComponent*>& Components, int32 UpdateCounter);
+	void ForceUpdateVisibilityOrDeferTask(const TArray<FFastGeoPrimitiveComponent*>& Components, int32 UpdateCounter);
 	void UpdateVisibility_Internal(TArray<FFastGeoPrimitiveComponent*>&& ShowComponents, TArray<FFastGeoPrimitiveComponent*>&& HideComponents, int32 UpdateCounter = 0);
 	void InitializeDynamicProperties();
 	void SetOwnerContainer(UFastGeoContainer* InOwner);
diff --git a/Engine/Plugins/Experimental/FastGeoStreaming/Source/FastGeoStreaming/Private/FastGeoComponentCluster.cpp b/Engine/Plugins/Experimental/FastGeoStreaming/Source/FastGeoStreaming/Private/FastGeoComponentCluster.cpp
--- a/Engine/Plugins/Experimental/FastGeoStreaming/Source/FastGeoStreaming/Private/FastGeoComponentCluster.cpp
+++ b/Engine/Plugins/Experimental/FastGeoStreaming/Source/FastGeoStreaming/Private/FastGeoComponentCluster.cpp
@@ -5,6 +5,7 @@
 #include "SceneInterface.h"
 #include "PrimitiveSceneInfo.h"
 #include "FastGeoWeakElement.h"
+#include "Containers/Ticker.h"
 
 const FFastGeoElementType FFastGeoComponentCluster::Type(&IFastGeoElement::Type);
 
@@ -191,6 +192,27 @@ void FFastGeoComponentCluster::ForceUpdateVisibility(const TArray<FFastGeoPrimit
 	UpdateVisibility_Internal(MoveTemp(ShowComponents), MoveTemp(HideComponents), ++UpdateCounter);
 }
 
+void FFastGeoComponentCluster::ForceUpdateVisibilityOrDeferTask(const TArray<FFastGeoPrimitiveComponent*>& Components, int32 UpdateCounter)
+{
+	if (UpdateCounter <= 1)
+	{
+		ForceUpdateVisibility(Components, UpdateCounter);
+		return;
+	}
+
+	// If primitives are not yet registered in the scene, immediate retries during blocking loads can loop indefinitely.
+	// Defer to the next frame to allow pending AddPrimitive render commands to complete.
+	FWeakFastGeoComponentCluster ClusterWeak(this);
+	FTSTicker::GetCoreTicker().AddTicker(FTickerDelegate::CreateLambda([ClusterWeak, Components, UpdateCounter](float) -> bool
+	{
+		if (FFastGeoComponentCluster* Cluster = ClusterWeak.Get())
+		{
+			Cluster->ForceUpdateVisibility(Components, UpdateCounter);
+		}
+		return false;
+	}));
+}
+
 void FFastGeoComponentCluster::UpdateVisibility_Internal(TArray<FFastGeoPrimitiveComponent*>&& ShowComponents, TArray<FFastGeoPrimitiveComponent*>&& HideComponents, int32 UpdateCounter)
 {
 	if (!ShowComponents.IsEmpty() || !HideComponents.IsEmpty())
@@ -236,7 +258,7 @@ void FFastGeoComponentCluster::UpdateVisibility_Internal(TArray<FFastGeoPrimitiv
 				{
 					if (FFastGeoComponentCluster* Cluster = ClusterWeak.Get())
 					{
-						Cluster->ForceUpdateVisibility(NotReadyComponents, UpdateCounter);
+						Cluster->ForceUpdateVisibilityOrDeferTask(NotReadyComponents, UpdateCounter);
 					}
 				}, LowLevelTasks::ETaskPriority::Normal, UE::Tasks::EExtendedTaskPriority::GameThreadNormalPri);
 			}
