[FastGeo] Refactor transformer debug logging into structured reports
- Refactored debug output to collect per-actor/component results into a structured report, replacing inline logging.
- Debug report is now shared between Debug Mode (cook/PIE) and Debug Selection Mode (editor), with contextual titles.
- Report groups actors by category (Ignored/Rejected/Allowed) with counts and component tree under allowed actors.
- Differentiate Discard vs Reject in transformation result logs.
- Log always-ignored actors in the debug report.
- Skip debug selection mode when FastGeo.Enable=0 or FastGeo.EnableTransformer=0, with reason in the log.

#Synced-CL 51624647

diff --git a/Engine/Plugins/Experimental/FastGeoStreaming/Source/FastGeoStreaming/Private/FastGeoWorldPartitionRuntimeCellTransformer.cpp b/Engine/Plugins/Experimental/FastGeoStreaming/Source/FastGeoStreaming/Private/FastGeoWorldPartitionRuntimeCellTransformer.cpp
--- a/Engine/Plugins/Experimental/FastGeoStreaming/Source/FastGeoStreaming/Private/FastGeoWorldPartitionRuntimeCellTransformer.cpp
+++ b/Engine/Plugins/Experimental/FastGeoStreaming/Source/FastGeoStreaming/Private/FastGeoWorldPartitionRuntimeCellTransformer.cpp
@@ -70,6 +70,16 @@ namespace FastGeo
 	// Tag use to force exclude actors from FastGeoStreaming
 	static const FName NAME_NoFastGeo(TEXT("NoFastGeo"));
 
+	static constexpr const TCHAR* LogSeparatorOuter = TEXT("========================================================================");
+	static constexpr const TCHAR* LogSeparator = TEXT("------------------------------------------------------------------------");
+
+	static void LogDebugSelectionMessage(const TCHAR* Message)
+	{
+		UE_LOG(LogFastGeoStreaming, Log, TEXT("%s"), LogSeparatorOuter);
+		UE_LOG(LogFastGeoStreaming, Log, TEXT("FastGeo Debug Selection Mode: %s"), Message);
+		UE_LOG(LogFastGeoStreaming, Log, TEXT("%s"), LogSeparatorOuter);
+	}
+
 	static IWorldPartitionHLODObject* GetHLODObject(AActor* Actor)
 	{
 		if (AWorldPartitionHLOD* HLODActor = Cast<AWorldPartitionHLOD>(Actor))
@@ -249,11 +259,24 @@ bool UFastGeoWorldPartitionRuntimeCellTransformer::IsDebugMode() const
 
 void UFastGeoWorldPartitionRuntimeCellTransformer::OnSelectionChanged(UObject* Object)
 {
-	if (bDebugModeOnSelection && IsEnabled())
+	if (bDebugSelectionMode && IsEnabled())
 	{
+		if (!FFastGeoStreamingModule::IsFastGeoEnabled())
+		{
+			FastGeo::LogDebugSelectionMessage(TEXT("Skipped - FastGeo.Enable is false."));
+			return;
+		}
+
+		if (!UFastGeoWorldPartitionRuntimeCellTransformer::IsFastGeoTransformerEnabled)
+		{
+			FastGeo::LogDebugSelectionMessage(TEXT("Skipped - FastGeo.EnableTransformer is false."));
+			return;
+		}
+
 		TArray<AActor*> SelectedActors;
+		TArray<AActor*> IgnoredActors;
 
-		TFunction<void(AActor* InActor)> AddActorToSelection = [this, &SelectedActors, &AddActorToSelection](AActor* InActor)
+		TFunction<void(AActor* InActor)> AddActorToSelection = [this, &SelectedActors, &IgnoredActors, &AddActorToSelection](AActor* InActor)
 		{
 			if (!CanAlwaysIgnoreActor(InActor))
 			{
@@ -274,6 +297,10 @@ void UFastGeoWorldPartitionRuntimeCellTransformer::OnSelectionChanged(UObject* O
 					SelectedActors.Add(InActor);
 				}
 			}
+			else if (IsValid(InActor))
+			{
+				IgnoredActors.Add(InActor);
+			}
 		};
 
 		if (USelection* Selection = Cast<USelection>(Object))
@@ -287,27 +314,20 @@ void UFastGeoWorldPartitionRuntimeCellTransformer::OnSelectionChanged(UObject* O
 			}
 		}
 
-		if (!SelectedActors.IsEmpty())
+		if (!SelectedActors.IsEmpty() || !IgnoredActors.IsEmpty())
 		{
-			TGuardValue<bool> Guard(FFastGeoTransformResult::ShouldReport, true);
-
-			UE_LOG(LogFastGeoStreaming, Log, TEXT("------------------------------------------------------------------------"));
-			UE_LOG(LogFastGeoStreaming, Log, TEXT("- FastGeoStreaming Debug Mode: Transformation on %d selected actors "), SelectedActors.Num());
+			TGuardValue<bool> Guard(FFastGeoTransformResult::bShouldCollectReasons, true);
 
 			FTransformationStats Stats;
+			TArray<FActorReport> ActorReports;
 			TArray<TPair<AActor*, FTransformableActor>> TransformableActors;
 
-			GatherTransformableActors(SelectedActors, SelectedActors[0]->GetLevel(), TransformableActors, Stats);
-
-			Stats.DumpStats(TEXT("  - "));
-
-			if (!FPhysScene::SupportsAsyncPhysicsStateCreation() ||
-				!FPhysScene::SupportsAsyncPhysicsStateDestruction())
+			if (!SelectedActors.IsEmpty())
 			{
-				UE_LOG(LogFastGeoStreaming, Warning, TEXT(" - NOTE: FastGeoStreaming requires 'p.Chaos.EnableAsyncInitBody' to be enabled."));
+				GatherTransformableActors(SelectedActors, SelectedActors[0]->GetLevel(), TransformableActors, Stats, &ActorReports);
 			}
 
-			UE_LOG(LogFastGeoStreaming, Log, TEXT("------------------------------------------------------------------------"));
+			DumpDebugReport(TEXT("FastGeo Debug Selection Mode"), IgnoredActors, ActorReports, Stats);
 		}
 	}
 }
@@ -323,14 +343,15 @@ bool UFastGeoWorldPartitionRuntimeCellTransformer::CanAlwaysIgnoreActor(AActor*
 
 void UFastGeoWorldPartitionRuntimeCellTransformer::Transform(ULevel* InLevel)
 {
-	TGuardValue<bool> Guard(FFastGeoTransformResult::ShouldReport, IsDebugMode());
+	const bool bIsDebugMode = IsDebugMode();
+	TGuardValue<bool> Guard(FFastGeoTransformResult::bShouldCollectReasons, bIsDebugMode);
 
 	if (!FFastGeoStreamingModule::IsFastGeoEnabled() || !UFastGeoWorldPartitionRuntimeCellTransformer::IsFastGeoTransformerEnabled)
 	{
 		return;
 	}
 
-	if (!FPhysScene::SupportsAsyncPhysicsStateCreation() || 
+	if (!FPhysScene::SupportsAsyncPhysicsStateCreation() ||
 		!FPhysScene::SupportsAsyncPhysicsStateDestruction())
 	{
 		UE_LOG(LogFastGeoStreaming, Error, TEXT("FastGeoStreaming Cell Transformer requires 'p.Chaos.EnableAsyncInitBody' to be enabled."));
@@ -340,16 +361,11 @@ void UFastGeoWorldPartitionRuntimeCellTransformer::Transform(ULevel* InLevel)
 	check(InLevel);
 	check(!UFastGeoWorldPartitionRuntimeCellTransformer::CurrentTransformer);
 	TGuardValue<UFastGeoWorldPartitionRuntimeCellTransformer*> GuardCurrentTransformer(UFastGeoWorldPartitionRuntimeCellTransformer::CurrentTransformer, this);
-	
-	if (IsDebugMode())
-	{
-		UE_LOG(LogFastGeoStreaming, Log, TEXT("------------------------------------------------------------------------"));
-		UE_LOG(LogFastGeoStreaming, Log, TEXT("- FastGeoStreaming Debug Mode: Transforming Level '%s'"), *InLevel->GetPathName());
-	}
 
 	FTransformationStats Stats;
 	TArray<TPair<AActor*, FTransformableActor>> TransformableActors;
-	GatherTransformableActors(InLevel->Actors, InLevel, TransformableActors, Stats);
+	TArray<FActorReport> ActorReports;
+	GatherTransformableActors(InLevel->Actors, InLevel, TransformableActors, Stats, bIsDebugMode ? &ActorReports : nullptr);
 
 	if (!TransformableActors.IsEmpty())
 	{
@@ -482,11 +498,11 @@ void UFastGeoWorldPartitionRuntimeCellTransformer::Transform(ULevel* InLevel)
 
 	InLevel->Actors.Remove(nullptr);
 
-	if (IsDebugMode())
+	if (bIsDebugMode)
 	{
-		UE_LOG(LogFastGeoStreaming, Log, TEXT("- Transformation result of Level '%s'"), *InLevel->GetPathName());
-		Stats.DumpStats(TEXT("  - "));
-		UE_LOG(LogFastGeoStreaming, Log, TEXT("------------------------------------------------------------------------"));
+		const FString Title = FString::Printf(TEXT("FastGeo Debug Mode: Level '%s'"), *InLevel->GetPathName());
+		TArray<AActor*> EmptyIgnored;
+		DumpDebugReport(*Title, EmptyIgnored, ActorReports, Stats);
 	}
 }
 
@@ -607,11 +623,12 @@ TMap<AActor*, TArray<AActor*>> UFastGeoWorldPartitionRuntimeCellTransformer::Bui
 	return ReferencedActors;
 }
 
-void UFastGeoWorldPartitionRuntimeCellTransformer::GatherTransformableActors(const TArray<AActor*>& InActors, const ULevel* InLevel, TArray<TPair<AActor*, FTransformableActor>>& OutTransformableActors, FTransformationStats& OutStats)
+void UFastGeoWorldPartitionRuntimeCellTransformer::GatherTransformableActors(const TArray<AActor*>& InActors, const ULevel* InLevel, TArray<TPair<AActor*, FTransformableActor>>& OutTransformableActors, FTransformationStats& OutStats, TArray<FActorReport>* OutReports)
 {
 	TMap<AActor*, FTransformableActor> TransformableActors;
 	// Get transformation result for each actor
 	// This will retrieve the transformable components, and whether the actor is fully transformable (ie. if the actor can be deleted)
+	TMap<const AActor*, int32> ActorReportIndices;
 	for (int32 ActorIndex = 0; ActorIndex < InActors.Num(); ++ActorIndex)
 	{
 		AActor* Actor = InActors[ActorIndex];
@@ -620,9 +637,26 @@ void UFastGeoWorldPartitionRuntimeCellTransformer::GatherTransformableActors(con
 			OutStats.TotalActorCount++;
 			OutStats.TotalComponentCount += Actor->GetComponents().Num();
 
+			TArray<FComponentReport>* ComponentReports = nullptr;
+			if (OutReports)
+			{
+				ActorReportIndices.Add(Actor, OutReports->Num());
+				FActorReport& ActorReport = OutReports->AddDefaulted_GetRef();
+				ActorReport.Actor = Actor;
+				ComponentReports = &ActorReport.ComponentReports;
+			}
+
 			bool bIsActorFullyTransformable;
 			TArray<UActorComponent*> TransformableComponents;
-			FFastGeoTransformResult ActorTransformResult = CanTransformActor(Actor, bIsActorFullyTransformable, TransformableComponents);
+			FFastGeoTransformResult ActorTransformResult = CanTransformActor(Actor, bIsActorFullyTransformable, TransformableComponents, ComponentReports);
+
+			if (OutReports)
+			{
+				FActorReport& ActorReport = (*OutReports)[ActorReportIndices[Actor]];
+				ActorReport.Result = ActorTransformResult.GetResult();
+				ActorReport.Reason = ActorTransformResult.GetReason();
+			}
+
 			if (ActorTransformResult.GetResult() == EFastGeoTransform::Allow)
 			{
 				// Sort components by pathname to ensure deterministic order
@@ -658,13 +692,16 @@ void UFastGeoWorldPartitionRuntimeCellTransformer::GatherTransformableActors(con
 			// If one of the referencer is not fully transformed to FastGeo
 			if (bHasNonFastGeoReferencer)
 			{
-				if (TransformableReferencer)
-				{
-					FFastGeoTransformResult(EFastGeoTransform::Reject, [Actor, TransformableReferencer, &InActors] { return FString::Printf(TEXT("Actor '%s' is referenced by a non-fully transformed actor ('%s')"), *Actor->GetName(), *InActors[TransformableReferencer->ActorIndex]->GetName()); });
-				}
-				else
+				FFastGeoTransformResult RejectResult = TransformableReferencer
+					? FFastGeoTransformResult(EFastGeoTransform::Reject, [Actor, TransformableReferencer, &InActors] { return FString::Printf(TEXT("Actor '%s' is referenced by a non-fully transformed actor ('%s')"), *Actor->GetName(), *InActors[TransformableReferencer->ActorIndex]->GetName()); })
+					: FFastGeoTransformResult(EFastGeoTransform::Reject, [Actor, &Referencer] { return FString::Printf(TEXT("Actor '%s' is referenced by a non transformed actor ('%s')"), *Actor->GetName(), *Referencer->GetName()); });
+
+				if (const int32* ReportIndex = ActorReportIndices.Find(Actor))
 				{
-					FFastGeoTransformResult(EFastGeoTransform::Reject, [Actor, &Referencer] { return FString::Printf(TEXT("Actor '%s' is referenced by a non transformed actor ('%s')"), *Actor->GetName(), *Referencer->GetName()); });
+					FActorReport& Report = (*OutReports)[*ReportIndex];
+					Report.Result = RejectResult.GetResult();
+					Report.Reason = RejectResult.GetReason();
+					Report.ComponentReports.Empty();
 				}
 
 				It.RemoveCurrent();
@@ -689,7 +726,106 @@ void UFastGeoWorldPartitionRuntimeCellTransformer::GatherTransformableActors(con
 	});
 }
 
-void UFastGeoWorldPartitionRuntimeCellTransformer::FTransformationStats::DumpStats(const TCHAR* InPrefixString)
+void UFastGeoWorldPartitionRuntimeCellTransformer::DumpDebugReport(const TCHAR* Title, const TArray<AActor*>& IgnoredActors, const TArray<FActorReport>& ActorReports, const FTransformationStats& Stats)
+{
+	const int32 TotalSelected = IgnoredActors.Num() + Stats.TotalActorCount;
+
+	UE_LOG(LogFastGeoStreaming, Log, TEXT("%s"), FastGeo::LogSeparatorOuter);
+
+	if (IgnoredActors.Num() > 0 && Stats.TotalActorCount == 0)
+	{
+		UE_LOG(LogFastGeoStreaming, Log, TEXT("%s: %d actors (all always ignored)"), Title, IgnoredActors.Num());
+	}
+	else if (IgnoredActors.Num() > 0)
+	{
+		UE_LOG(LogFastGeoStreaming, Log, TEXT("%s: %d actors (%d always ignored)"), Title, TotalSelected, IgnoredActors.Num());
+	}
+	else
+	{
+		UE_LOG(LogFastGeoStreaming, Log, TEXT("%s: %d actors"), Title, TotalSelected);
+	}
+
+	UE_LOG(LogFastGeoStreaming, Log, TEXT("%s"), FastGeo::LogSeparator);
+
+	// Count entries per category
+	int32 NumRejected = 0;
+	int32 NumAllowed = 0;
+	for (const FActorReport& Report : ActorReports)
+	{
+		if (Report.Result == EFastGeoTransform::Allow)
+		{
+			NumAllowed++;
+		}
+		else
+		{
+			NumRejected++;
+		}
+	}
+
+	auto LogSectionHeader = [](const TCHAR* Title, int32 Count)
+	{
+		UE_LOG(LogFastGeoStreaming, Log, TEXT(""));
+		UE_LOG(LogFastGeoStreaming, Log, TEXT("%s"), FastGeo::LogSeparator);
+		UE_LOG(LogFastGeoStreaming, Log, TEXT("%s (%d):"), Title, Count);
+		UE_LOG(LogFastGeoStreaming, Log, TEXT("%s"), FastGeo::LogSeparator);
+	};
+
+	// Ignored
+	if (IgnoredActors.Num() > 0)
+	{
+		LogSectionHeader(TEXT("Ignored"), IgnoredActors.Num());
+		for (const AActor* Actor : IgnoredActors)
+		{
+			UE_LOG(LogFastGeoStreaming, Log, TEXT("    * '%s' (%s)"), *Actor->GetActorNameOrLabel(), *Actor->GetClass()->GetName());
+		}
+	}
+
+	// Rejected/Discarded
+	if (NumRejected > 0)
+	{
+		LogSectionHeader(TEXT("Rejected"), NumRejected);
+		for (const FActorReport& Report : ActorReports)
+		{
+			UE_CLOG(Report.Result != EFastGeoTransform::Allow, LogFastGeoStreaming, Log, TEXT("    * '%s' - %s"), *Report.Actor->GetActorNameOrLabel(), *Report.Reason);
+		}
+	}
+
+	// Allowed
+	if (NumAllowed > 0)
+	{
+		LogSectionHeader(TEXT("Allowed"), NumAllowed);
+		for (const FActorReport& Report : ActorReports)
+		{
+			if (Report.Result == EFastGeoTransform::Allow)
+			{
+				const FString ActorName = Report.Actor->GetActorNameOrLabel();
+				UE_LOG(LogFastGeoStreaming, Log, TEXT("    * '%s' (%d components)"), *ActorName, Report.ComponentReports.Num());
+				// Allowed components first, then discarded
+				for (const FComponentReport& Comp : Report.ComponentReports)
+				{
+					UE_CLOG(Comp.Result == EFastGeoTransform::Allow, LogFastGeoStreaming, Log, TEXT("        - Allow: '%s'"), *Comp.Component->GetName());
+				}
+				for (const FComponentReport& Comp : Report.ComponentReports)
+				{
+					UE_CLOG(Comp.Result != EFastGeoTransform::Allow, LogFastGeoStreaming, Log, TEXT("        - Discard: '%s'%s"), *Comp.Component->GetName(), *FString(Comp.Reason.IsEmpty() ? TEXT("") : *FString::Printf(TEXT(" - %s"), *Comp.Reason)));
+				}
+			}
+		}
+	}
+
+	UE_LOG(LogFastGeoStreaming, Log, TEXT("%s"), FastGeo::LogSeparator);
+
+	Stats.DumpStats(TEXT("  -"));
+
+	if (!FPhysScene::SupportsAsyncPhysicsStateCreation() || !FPhysScene::SupportsAsyncPhysicsStateDestruction())
+	{
+		UE_LOG(LogFastGeoStreaming, Warning, TEXT("  - NOTE: FastGeoStreaming requires 'p.Chaos.EnableAsyncInitBody' to be enabled."));
+	}
+
+	UE_LOG(LogFastGeoStreaming, Log, TEXT("%s"), FastGeo::LogSeparatorOuter);
+}
+
+void UFastGeoWorldPartitionRuntimeCellTransformer::FTransformationStats::DumpStats(const TCHAR* InPrefixString) const
 {
 	if (TotalActorCount)
 	{
@@ -797,7 +933,7 @@ bool UFastGeoWorldPartitionRuntimeCellTransformer::IsComponentMobilitySupported(
 	return bIsSupported;
 }
 
-FFastGeoTransformResult UFastGeoWorldPartitionRuntimeCellTransformer::CanTransformActor(AActor* InActor, bool& bOutIsActorFullyTransformable, TArray<UActorComponent*>& OutTransformableComponents) const
+FFastGeoTransformResult UFastGeoWorldPartitionRuntimeCellTransformer::CanTransformActor(AActor* InActor, bool& bOutIsActorFullyTransformable, TArray<UActorComponent*>& OutTransformableComponents, TArray<FComponentReport>* OutComponentReports) const
 {
 	bOutIsActorFullyTransformable = false;
 
@@ -808,17 +944,17 @@ FFastGeoTransformResult UFastGeoWorldPartitionRuntimeCellTransformer::CanTransfo
 	}
 
 	FString Reason;
-	
+
 	if (InActor->ActorHasTag(NAME_CellTransformerIgnoreActor))
 	{
 		return { EFastGeoTransform::Reject, [InActor] { return FString::Printf(TEXT("Actor %s is tagged '%s'"), *InActor->GetName(), *NAME_CellTransformerIgnoreActor.ToString()); } };
 	}
-	
+
 	if (InActor->ActorHasTag(FastGeo::NAME_NoFastGeo))
 	{
 		return { EFastGeoTransform::Reject, [InActor] { return FString::Printf(TEXT("Actor %s is tagged '%s'"), *InActor->GetName(), *FastGeo::NAME_NoFastGeo.ToString()); } };
 	}
-	
+
 	if (!IsActorTransformable(InActor, Reason))
 	{
 		return { EFastGeoTransform::Reject, [InActor, Reason] { return FString::Printf(TEXT("Actor %s [%s]"), *InActor->GetName(), *Reason); } };
@@ -876,10 +1012,14 @@ FFastGeoTransformResult UFastGeoWorldPartitionRuntimeCellTransformer::CanTransfo
 
 	// Gather transformable components
 	TArray<UActorComponent*> TransformResults[EnumToIndex(EFastGeoTransform::MAX)];
-	InActor->ForEachComponent<UActorComponent>(false, [this, &TransformResults](UActorComponent* Component)
+	InActor->ForEachComponent<UActorComponent>(false, [this, &TransformResults, OutComponentReports](UActorComponent* Component)
 	{
 		FFastGeoTransformResult TransformComponentResult = CanTransformComponent(Component);
 		TransformResults[TransformComponentResult.GetResultIndex()].Add(Component);
+		if (OutComponentReports)
+		{
+			OutComponentReports->Add({ Component, TransformComponentResult.GetResult(), TransformComponentResult.GetReason() });
+		}
 	});
 
 	if (TransformResults[EnumToIndex(EFastGeoTransform::Allow)].IsEmpty())
@@ -890,19 +1030,21 @@ FFastGeoTransformResult UFastGeoWorldPartitionRuntimeCellTransformer::CanTransfo
 	// If actor contains only discardable or transformable components, we can actually get rid of it
 	bOutIsActorFullyTransformable = TransformResults[EnumToIndex(EFastGeoTransform::Reject)].IsEmpty();
 
-	Reason.Reset();
-	if (!IsFullyTransformedActorDeletable(InActor, Reason))
+	// BP actors can't be partially transformed (rerun CS will be called in PIE when registering
+	// the component and also called when registering components during cook/save of the level).
+	if (UBlueprint::GetBlueprintFromClass(InActor->GetClass()) && !bOutIsActorFullyTransformable)
 	{
-		bOutIsActorFullyTransformable = false;
-		UE_CLOG(FFastGeoTransformResult::ShouldReport, LogFastGeoStreaming, Log, TEXT("  * Can't delete fully transformed actor %s [%s]"), *InActor->GetName(), *Reason);
+		return { EFastGeoTransform::Reject, [InActor] { return FString::Printf(TEXT("Blueprint actor %s has non-transformable components"), *InActor->GetName()); } };
 	}
 
-	// Can't convert partially a BP actors
-	// Rerun CS will be called in PIE when registering the component and
-	// also called when registering components during cook/save of the level.
-	if (UBlueprint::GetBlueprintFromClass(InActor->GetClass()) && !bOutIsActorFullyTransformable)
+	Reason.Reset();
+	if (!IsFullyTransformedActorDeletable(InActor, Reason))
 	{
-		return { EFastGeoTransform::Reject, [InActor] { return FString::Printf(TEXT("Actor %s is a Blueprint and can't be fully transformed."), *InActor->GetName()); } };
+		bOutIsActorFullyTransformable = false;
+		if (UBlueprint::GetBlueprintFromClass(InActor->GetClass()))
+		{
+			return { EFastGeoTransform::Reject, [InActor, Reason] { return FString::Printf(TEXT("Blueprint actor %s can't be deleted [%s]"), *InActor->GetName(), *Reason); } };
+		}
 	}
 
 	OutTransformableComponents = MoveTemp(TransformResults[EnumToIndex(EFastGeoTransform::Allow)]);
@@ -1143,29 +1285,29 @@ void UFastGeoWorldPartitionRuntimeCellTransformer::PostEditChangeProperty(FPrope
 {
 	const FName PropertyName = InPropertyChangedEvent.GetPropertyName();
 	if ((PropertyName == GET_MEMBER_NAME_CHECKED(UFastGeoWorldPartitionRuntimeCellTransformer, bDebugMode) || 
-		 PropertyName == GET_MEMBER_NAME_CHECKED(UFastGeoWorldPartitionRuntimeCellTransformer, bDebugModeOnSelection)) && !FastGeo::GPackageWasDirty)
+		 PropertyName == GET_MEMBER_NAME_CHECKED(UFastGeoWorldPartitionRuntimeCellTransformer, bDebugSelectionMode)) && !FastGeo::GPackageWasDirty)
 	{
 		GetPackage()->ClearDirtyFlag();
 	}
 }
 
-bool FFastGeoTransformResult::ShouldReport = false;
+bool FFastGeoTransformResult::bShouldCollectReasons = false;
 
 FFastGeoTransformResult::FFastGeoTransformResult(EFastGeoTransform InTransformResult, const TCHAR* InFailureReason)
 	: TransformResult(InTransformResult)
 {
-	if (TransformResult != EFastGeoTransform::Allow && InFailureReason && ShouldReport)
+	if (TransformResult != EFastGeoTransform::Allow && InFailureReason && bShouldCollectReasons)
 	{
-		UE_LOG(LogFastGeoStreaming, Log, TEXT("  * Can't transform: %s"), InFailureReason);
+		Reason = InFailureReason;
 	}
 }
 
 FFastGeoTransformResult::FFastGeoTransformResult(EFastGeoTransform InTransformResult, TFunctionRef<FString()> InFailureReasonFunc)
 	: TransformResult(InTransformResult)
 {
-	if (TransformResult != EFastGeoTransform::Allow && ShouldReport)
+	if (TransformResult != EFastGeoTransform::Allow && bShouldCollectReasons)
 	{
-		UE_LOG(LogFastGeoStreaming, Log, TEXT("  * Can't transform: %s"), *InFailureReasonFunc());
+		Reason = InFailureReasonFunc();
 	}
 }
 
diff --git a/Engine/Plugins/Experimental/FastGeoStreaming/Source/FastGeoStreaming/Public/FastGeoWorldPartitionRuntimeCellTransformer.h b/Engine/Plugins/Experimental/FastGeoStreaming/Source/FastGeoStreaming/Public/FastGeoWorldPartitionRuntimeCellTransformer.h
--- a/Engine/Plugins/Experimental/FastGeoStreaming/Source/FastGeoStreaming/Public/FastGeoWorldPartitionRuntimeCellTransformer.h
+++ b/Engine/Plugins/Experimental/FastGeoStreaming/Source/FastGeoStreaming/Public/FastGeoWorldPartitionRuntimeCellTransformer.h
@@ -25,7 +25,7 @@ inline constexpr uint32 EnumToIndex(EFastGeoTransform Value)
 	return static_cast<uint32>(Value);
 }
 
-/** Used to conditionally report the error result when not a success. */
+/** Stores the transformation result and optionally collects the failure reason for debug reports. */
 struct FFastGeoTransformResult
 {
 public:
@@ -34,11 +34,13 @@ public:
 
 	EFastGeoTransform GetResult() const { return TransformResult; }
 	uint32 GetResultIndex() const { return EnumToIndex(TransformResult); }
+	const FString& GetReason() const { return Reason; }
 
-	static bool ShouldReport;
+	static bool bShouldCollectReasons;
 
 private:
 	EFastGeoTransform TransformResult;
+	FString Reason;
 };
 
 /** Defines the memory storage format for FastGeo instance transforms. */
@@ -196,15 +198,6 @@ protected:
 	/** Whether component mobility is supported by the transformer. */
 	bool IsComponentMobilitySupported(USceneComponent* InComponent, FString& OutReason) const;
 
-	FFastGeoTransformResult CanTransformActor(AActor* InActor, bool& bOutIsActorFullyTransformable, TArray<UActorComponent*>& OutTransformableComponents) const;
-	FFastGeoTransformResult IsAllowedActorClass(AActor* InActor) const;
-	FFastGeoTransformResult CanTransformComponent(UActorComponent* InComponent) const;
-	FFastGeoTransformResult IsAllowedComponentClass(UActorComponent* InComponent) const;
-	bool CanRemoveActor(AActor* InActor, TSet<UActorComponent*>& InIgnoredComponents) const;
-	bool CanAlwaysIgnoreActor(AActor* InActor) const;
-	void OnSelectionChanged(UObject* Object);
-	bool IsDebugMode() const;
-
 	struct FTransformableActor
 	{
 		int32 ActorIndex = 0;
@@ -220,11 +213,36 @@ protected:
 		int32 PartiallyTransformableActorCount = 0;
 		int32 TransformedComponentCount = 0;
 
-		void DumpStats(const TCHAR* InPrefixString);
+		void DumpStats(const TCHAR* InPrefixString) const;
+	};
+
+	struct FComponentReport
+	{
+		UActorComponent* Component = nullptr;
+		EFastGeoTransform Result = EFastGeoTransform::Allow;
+		FString Reason;
+	};
+
+	struct FActorReport
+	{
+		AActor* Actor = nullptr;
+		EFastGeoTransform Result = EFastGeoTransform::Allow;
+		FString Reason;
+		TArray<FComponentReport> ComponentReports;
 	};
 
+	FFastGeoTransformResult CanTransformActor(AActor* InActor, bool& bOutIsActorFullyTransformable, TArray<UActorComponent*>& OutTransformableComponents, TArray<FComponentReport>* OutComponentReports = nullptr) const;
+	FFastGeoTransformResult IsAllowedActorClass(AActor* InActor) const;
+	FFastGeoTransformResult CanTransformComponent(UActorComponent* InComponent) const;
+	FFastGeoTransformResult IsAllowedComponentClass(UActorComponent* InComponent) const;
+	bool CanRemoveActor(AActor* InActor, TSet<UActorComponent*>& InIgnoredComponents) const;
+	bool CanAlwaysIgnoreActor(AActor* InActor) const;
+	void OnSelectionChanged(UObject* Object);
+	bool IsDebugMode() const;
+
 	TMap<AActor*, TArray<AActor*>> BuildActorsReferencesMap(const TArray<AActor*>& InActors);
-	void GatherTransformableActors(const TArray<AActor*>& InActors, const ULevel* InLevel, TArray<TPair<AActor*, FTransformableActor>>& OutTransformableActors, FTransformationStats& OutStats);
+	void GatherTransformableActors(const TArray<AActor*>& InActors, const ULevel* InLevel, TArray<TPair<AActor*, FTransformableActor>>& OutTransformableActors, FTransformationStats& OutStats, TArray<FActorReport>* OutReports = nullptr);
+	static void DumpDebugReport(const TCHAR* Title, const TArray<AActor*>& IgnoredActors, const TArray<FActorReport>& ActorReports, const FTransformationStats& Stats);
 #endif
 
 public:
@@ -248,7 +266,7 @@ protected:
 	bool bDebugMode;
 
 	UPROPERTY(EditAnywhere, Transient, SkipSerialization, Category = "Debug")
-	bool bDebugModeOnSelection;
+	bool bDebugSelectionMode;
 	
 	// The following properties are all deprecated in favor of the separate Settings property pointing to a reusable settings asset
 	UPROPERTY(meta = (DeprecatedProperty))
