[FastGeo] Fix false rejections in debug selection mode reference checking
Fix 1: Editor-only actors were included in the reference map, causing actors they reference to be falsely rejected. These actors are now skipped.
Fix 2: Actors referenced by non-selected but transformable actors were falsely rejected. Referencers are now lazily evaluated on-demand and cached.

#Synced-CL 51624647

diff --git a/Engine/Plugins/Experimental/FastGeoStreaming/Source/FastGeoStreaming/Private/FastGeoWorldPartitionRuntimeCellTransformer.cpp b/Engine/Plugins/Experimental/FastGeoStreaming/Source/FastGeoStreaming/Private/FastGeoWorldPartitionRuntimeCellTransformer.cpp
--- a/Engine/Plugins/Experimental/FastGeoStreaming/Source/FastGeoStreaming/Private/FastGeoWorldPartitionRuntimeCellTransformer.cpp
+++ b/Engine/Plugins/Experimental/FastGeoStreaming/Source/FastGeoStreaming/Private/FastGeoWorldPartitionRuntimeCellTransformer.cpp
@@ -570,6 +570,12 @@ TMap<AActor*, TArray<AActor*>> UFastGeoWorldPartitionRuntimeCellTransformer::Bui
 			continue;
 		}
 
+		// Editor-only actors won't exist at cook time, so their references should not prevent other actors from being transformed.
+		if (ReferencingActor->IsEditorOnly())
+		{
+			continue;
+		}
+
 		VisitedObjects.Reset();
 		VisitedObjects.Add(ReferencingActor);
 
@@ -625,9 +631,36 @@ TMap<AActor*, TArray<AActor*>> UFastGeoWorldPartitionRuntimeCellTransformer::Bui
 
 void UFastGeoWorldPartitionRuntimeCellTransformer::GatherTransformableActors(const TArray<AActor*>& InActors, const ULevel* InLevel, TArray<TPair<AActor*, FTransformableActor>>& OutTransformableActors, FTransformationStats& OutStats, TArray<FActorReport>* OutReports)
 {
-	TMap<AActor*, FTransformableActor> TransformableActors;
-	// Get transformation result for each actor
-	// This will retrieve the transformable components, and whether the actor is fully transformable (ie. if the actor can be deleted)
+	struct FEvaluatedActor
+	{
+		bool bIsFullyTransformable = false;
+		FFastGeoTransformResult Result = EFastGeoTransform::Allow;
+		TOptional<FTransformableActor> TransformData;
+
+		bool IsTransformable() const { return TransformData.IsSet(); }
+		void RejectTransformation() { TransformData.Reset(); }
+	};
+
+	TMap<AActor*, FEvaluatedActor> EvaluatedActors;
+
+	// Evaluate an actor and cache the result. Returns the cached entry.
+	auto EvaluateActor = [this, &EvaluatedActors](AActor* InActor) -> FEvaluatedActor&
+	{
+		if (FEvaluatedActor* Cached = EvaluatedActors.Find(InActor))
+		{
+			return *Cached;
+		}
+		FEvaluatedActor& Entry = EvaluatedActors.Add(InActor);
+		TArray<UActorComponent*> TempComponents;
+		Entry.Result = CanTransformActor(InActor, Entry.bIsFullyTransformable, TempComponents);
+		if (Entry.Result.GetResult() != EFastGeoTransform::Allow)
+		{
+			Entry.bIsFullyTransformable = false;
+		}
+		return Entry;
+	};
+
+	// Evaluate input actors and collect transformation results
 	TMap<const AActor*, int32> ActorReportIndices;
 	for (int32 ActorIndex = 0; ActorIndex < InActors.Num(); ++ActorIndex)
 	{
@@ -657,44 +690,65 @@ void UFastGeoWorldPartitionRuntimeCellTransformer::GatherTransformableActors(con
 				ActorReport.Reason = ActorTransformResult.GetReason();
 			}
 
-			if (ActorTransformResult.GetResult() == EFastGeoTransform::Allow)
+			FEvaluatedActor& EvaluatedActor = EvaluatedActors.Add(Actor);
+			EvaluatedActor.Result = MoveTemp(ActorTransformResult);
+
+			if (EvaluatedActor.Result.GetResult() == EFastGeoTransform::Allow)
 			{
+				EvaluatedActor.bIsFullyTransformable = bIsActorFullyTransformable;
 				// Sort components by pathname to ensure deterministic order
 				TransformableComponents.Sort([](const UActorComponent& A, const UActorComponent& B)
 				{
 					return A.GetPathName() < B.GetPathName();
 				});
-				TransformableActors.FindOrAdd(Actor, { ActorIndex, bIsActorFullyTransformable, MoveTemp(TransformableComponents) });
+				EvaluatedActor.TransformData.Emplace(FTransformableActor{ ActorIndex, bIsActorFullyTransformable, MoveTemp(TransformableComponents) });
 			}
 		}
 	}
 
 	TMap<AActor*, TArray<AActor*>> ReferencedActors = BuildActorsReferencesMap(InLevel->Actors);
-	for (TMap<AActor*, FTransformableActor>::TIterator It = TransformableActors.CreateIterator(); It; ++It)
+
+	// Pre-evaluate all referencers so we don't mutate EvaluatedActors during iteration below
+	for (const TPair<AActor*, TArray<AActor*>>& RefPair : ReferencedActors)
+	{
+		for (AActor* Referencer : RefPair.Value)
+		{
+			EvaluateActor(Referencer);
+		}
+	}
+
+	for (TPair<AActor*, FEvaluatedActor>& Pair : EvaluatedActors)
 	{
-		const TPair<AActor*, FTransformableActor>& Entry = *It;
-		const AActor* Actor = Entry.Key;
+		AActor* Actor = Pair.Key;
+		FEvaluatedActor& EvaluatedActor = Pair.Value;
+
+		// Skip actors that were rejected, discarded, or only evaluated as referencers
+		if (!EvaluatedActor.IsTransformable())
+		{
+			continue;
+		}
 
 		// Exclude actors that have:
 		//  * Non FastGeo referencers
 		//  * FastGeo referencers that are going to be only partially transformed
 		if (const TArray<AActor*>* Referencers = ReferencedActors.Find(Actor))
 		{
-			const AActor* Referencer = nullptr;
-			const FTransformableActor* TransformableReferencer = nullptr;
-			bool bHasNonFastGeoReferencer = Algo::AnyOf(*Referencers, [&Referencer, &TransformableReferencer, &TransformableActors](const AActor* InReferencer)
+			AActor* const* FailingReferencerPtr = Referencers->FindByPredicate([&EvaluatedActors](AActor* InReferencer)
 			{
-				Referencer = InReferencer;
-				TransformableReferencer = TransformableActors.Find(InReferencer);
-				return TransformableReferencer == nullptr || !TransformableReferencer->bIsActorFullyTransformable;
+				return !EvaluatedActors.FindChecked(InReferencer).bIsFullyTransformable;
 			});
 
-			// If one of the referencer is not fully transformed to FastGeo
-			if (bHasNonFastGeoReferencer)
+			if (FailingReferencerPtr)
 			{
-				FFastGeoTransformResult RejectResult = TransformableReferencer
-					? FFastGeoTransformResult(EFastGeoTransform::Reject, [Actor, TransformableReferencer, &InActors] { return FString::Printf(TEXT("Actor '%s' is referenced by a non-fully transformed actor ('%s')"), *Actor->GetName(), *InActors[TransformableReferencer->ActorIndex]->GetName()); })
-					: FFastGeoTransformResult(EFastGeoTransform::Reject, [Actor, &Referencer] { return FString::Printf(TEXT("Actor '%s' is referenced by a non transformed actor ('%s')"), *Actor->GetName(), *Referencer->GetName()); });
+				AActor* FailingReferencer = *FailingReferencerPtr;
+				const FEvaluatedActor& ReferencerEval = EvaluatedActors.FindChecked(FailingReferencer);
+				const bool bReferencerIsTransformable = ReferencerEval.Result.GetResult() == EFastGeoTransform::Allow;
+				const FString& ReferencerReason = ReferencerEval.Result.GetReason();
+				FFastGeoTransformResult RejectResult(EFastGeoTransform::Reject, [Actor, FailingReferencer, bReferencerIsTransformable, &ReferencerReason]
+				{
+					const TCHAR* ReferencerDesc = bReferencerIsTransformable ? TEXT("non-fully transformed") : TEXT("non-transformable");
+					return FString::Printf(TEXT("Actor '%s' is referenced by a %s actor ('%s'%s)"), *Actor->GetName(), ReferencerDesc, *FailingReferencer->GetName(), *(ReferencerReason.IsEmpty() ? FString() : FString::Printf(TEXT(": %s"), *ReferencerReason)));
+				});
 
 				if (const int32* ReportIndex = ActorReportIndices.Find(Actor))
 				{
@@ -704,21 +758,25 @@ void UFastGeoWorldPartitionRuntimeCellTransformer::GatherTransformableActors(con
 					Report.ComponentReports.Empty();
 				}
 
-				It.RemoveCurrent();
+				EvaluatedActor.RejectTransformation();
 				continue;
 			}
 		}
 
-		OutStats.FullyTransformableActorCount += Entry.Value.bIsActorFullyTransformable ? 1 : 0;
-		OutStats.PartiallyTransformableActorCount += !Entry.Value.bIsActorFullyTransformable ? 1 : 0;
-		OutStats.TransformedComponentCount += Entry.Value.TransformableComponents.Num();
+		OutStats.FullyTransformableActorCount += EvaluatedActor.bIsFullyTransformable ? 1 : 0;
+		OutStats.PartiallyTransformableActorCount += !EvaluatedActor.bIsFullyTransformable ? 1 : 0;
+		OutStats.TransformedComponentCount += EvaluatedActor.TransformData->TransformableComponents.Num();
 	}
 
 	// Sort by pathname to ensure deterministic order
-	OutTransformableActors.Reserve(TransformableActors.Num());
-	for (const TPair<AActor*, FTransformableActor>& Pair : TransformableActors)
+	for (TPair<AActor*, FEvaluatedActor>& Pair : EvaluatedActors)
 	{
-		OutTransformableActors.Add(Pair);
+		AActor* Actor = Pair.Key;
+		FEvaluatedActor& EvaluatedActor = Pair.Value;
+		if (EvaluatedActor.IsTransformable())
+		{
+			OutTransformableActors.Emplace(Actor, MoveTemp(EvaluatedActor.TransformData.GetValue()));
+		}
 	}
 	OutTransformableActors.Sort([](const TPair<AActor*, FTransformableActor>& A, const TPair<AActor*, FTransformableActor>& B)
 	{
diff --git a/Engine/Plugins/Experimental/FastGeoStreaming/Source/FastGeoStreaming/Public/FastGeoWorldPartitionRuntimeCellTransformer.h b/Engine/Plugins/Experimental/FastGeoStreaming/Source/FastGeoStreaming/Public/FastGeoWorldPartitionRuntimeCellTransformer.h
--- a/Engine/Plugins/Experimental/FastGeoStreaming/Source/FastGeoStreaming/Public/FastGeoWorldPartitionRuntimeCellTransformer.h
+++ b/Engine/Plugins/Experimental/FastGeoStreaming/Source/FastGeoStreaming/Public/FastGeoWorldPartitionRuntimeCellTransformer.h
@@ -200,7 +200,7 @@ protected:
 
 	struct FTransformableActor
 	{
-		int32 ActorIndex = 0;
+		int32 ActorIndex = INDEX_NONE;
 		bool bIsActorFullyTransformable = false;
 		TArray<UActorComponent*> TransformableComponents;
 	};
