[FastGeo] Improve debug report component-level details
- Show component reports under rejected actors when available (e.g. "has no transformable components" now lists why each component was rejected or discarded).
- Distinguish Reject vs Discard at the component level in the report instead of labeling everything as Discard.
- Add reason strings to CanIgnoreComponent discard paths (e.g. "class is ignored (SceneComponent)").
- Fix typo in discard reason: "has is not visible" -> "is not visible".

#Synced-CL 51624647

diff --git a/Engine/Plugins/Experimental/FastGeoStreaming/Source/FastGeoStreaming/Private/FastGeoWorldPartitionRuntimeCellTransformer.cpp b/Engine/Plugins/Experimental/FastGeoStreaming/Source/FastGeoStreaming/Private/FastGeoWorldPartitionRuntimeCellTransformer.cpp
--- a/Engine/Plugins/Experimental/FastGeoStreaming/Source/FastGeoStreaming/Private/FastGeoWorldPartitionRuntimeCellTransformer.cpp
+++ b/Engine/Plugins/Experimental/FastGeoStreaming/Source/FastGeoStreaming/Private/FastGeoWorldPartitionRuntimeCellTransformer.cpp
@@ -750,9 +750,9 @@ void UFastGeoWorldPartitionRuntimeCellTransformer::GatherTransformableActors(con
 					return FString::Printf(TEXT("Actor '%s' is referenced by a %s actor ('%s'%s)"), *Actor->GetName(), ReferencerDesc, *FailingReferencer->GetName(), *(ReferencerReason.IsEmpty() ? FString() : FString::Printf(TEXT(": %s"), *ReferencerReason)));
 				});
 
-				if (const int32* ReportIndex = ActorReportIndices.Find(Actor))
+				if (OutReports)
 				{
-					FActorReport& Report = (*OutReports)[*ReportIndex];
+					FActorReport& Report = (*OutReports)[ActorReportIndices.FindChecked(Actor)];
 					Report.Result = RejectResult.GetResult();
 					Report.Reason = RejectResult.GetReason();
 					Report.ComponentReports.Empty();
@@ -838,13 +838,44 @@ void UFastGeoWorldPartitionRuntimeCellTransformer::DumpDebugReport(const TCHAR*
 		}
 	}
 
+	auto GetComponentResultPrefix = [](EFastGeoTransform Result) -> const TCHAR*
+	{
+		switch (Result)
+		{
+		case EFastGeoTransform::Allow:   return TEXT("Allow");
+		case EFastGeoTransform::Reject:  return TEXT("Reject");
+		case EFastGeoTransform::Discard: return TEXT("Discard");
+		default:                         return TEXT("Unknown");
+		}
+	};
+
+	auto LogComponentReports = [&GetComponentResultPrefix](const TArray<FComponentReport>& ComponentReports)
+	{
+		// Allowed components first, then rejected/discarded
+		for (const FComponentReport& Comp : ComponentReports)
+		{
+			UE_CLOG(Comp.Result == EFastGeoTransform::Allow, LogFastGeoStreaming, Log, TEXT("        - Allow: '%s'"), *Comp.Component->GetName());
+		}
+		for (const FComponentReport& Comp : ComponentReports)
+		{
+			UE_CLOG(Comp.Result != EFastGeoTransform::Allow, LogFastGeoStreaming, Log, TEXT("        - %s: '%s'%s"), GetComponentResultPrefix(Comp.Result), *Comp.Component->GetName(), *FString(Comp.Reason.IsEmpty() ? TEXT("") : *FString::Printf(TEXT(" - %s"), *Comp.Reason)));
+		}
+	};
+
 	// Rejected/Discarded
 	if (NumRejected > 0)
 	{
 		LogSectionHeader(TEXT("Rejected"), NumRejected);
 		for (const FActorReport& Report : ActorReports)
 		{
-			UE_CLOG(Report.Result != EFastGeoTransform::Allow, LogFastGeoStreaming, Log, TEXT("    * '%s' - %s"), *Report.Actor->GetActorNameOrLabel(), *Report.Reason);
+			if (Report.Result != EFastGeoTransform::Allow)
+			{
+				UE_LOG(LogFastGeoStreaming, Log, TEXT("    * '%s' - %s"), *Report.Actor->GetActorNameOrLabel(), *Report.Reason);
+				if (!Report.ComponentReports.IsEmpty())
+				{
+					LogComponentReports(Report.ComponentReports);
+				}
+			}
 		}
 	}
 
@@ -856,17 +887,8 @@ void UFastGeoWorldPartitionRuntimeCellTransformer::DumpDebugReport(const TCHAR*
 		{
 			if (Report.Result == EFastGeoTransform::Allow)
 			{
-				const FString ActorName = Report.Actor->GetActorNameOrLabel();
-				UE_LOG(LogFastGeoStreaming, Log, TEXT("    * '%s' (%d components)"), *ActorName, Report.ComponentReports.Num());
-				// Allowed components first, then discarded
-				for (const FComponentReport& Comp : Report.ComponentReports)
-				{
-					UE_CLOG(Comp.Result == EFastGeoTransform::Allow, LogFastGeoStreaming, Log, TEXT("        - Allow: '%s'"), *Comp.Component->GetName());
-				}
-				for (const FComponentReport& Comp : Report.ComponentReports)
-				{
-					UE_CLOG(Comp.Result != EFastGeoTransform::Allow, LogFastGeoStreaming, Log, TEXT("        - Discard: '%s'%s"), *Comp.Component->GetName(), *FString(Comp.Reason.IsEmpty() ? TEXT("") : *FString::Printf(TEXT(" - %s"), *Comp.Reason)));
-				}
+				UE_LOG(LogFastGeoStreaming, Log, TEXT("    * '%s' (%d components)"), *Report.Actor->GetActorNameOrLabel(), Report.ComponentReports.Num());
+				LogComponentReports(Report.ComponentReports);
 			}
 		}
 	}
@@ -1180,13 +1202,13 @@ FFastGeoTransformResult UFastGeoWorldPartitionRuntimeCellTransformer::CanTransfo
 
 	if (CanIgnoreComponent(InComponent))
 	{
-		return { EFastGeoTransform::Discard };
+		return { EFastGeoTransform::Discard, [InComponent] { return FString::Printf(TEXT("Component %s class is ignored (%s)"), *FastGeo::GetComponentShortName(InComponent), *InComponent->GetClass()->GetName()); } };
 	}
 
 	USceneComponent* SceneComponent = Cast<USceneComponent>(InComponent);
 	if (!SceneComponent)
 	{
-		return { EFastGeoTransform::Reject, [InComponent] { return FString::Printf(TEXT("Component %s class %s is not allowed"), *FastGeo::GetComponentShortName(InComponent), *InComponent->GetClass()->GetName()); } }; //-V522
+		return { EFastGeoTransform::Reject, [InComponent] { return FString::Printf(TEXT("Component %s class is not allowed (%s)"), *FastGeo::GetComponentShortName(InComponent), *InComponent->GetClass()->GetName()); } }; //-V522
 	}
 
 	const FFastGeoTransformResult AllowedComponentClassResult = IsAllowedComponentClass(InComponent);
@@ -1322,12 +1344,12 @@ FFastGeoTransformResult UFastGeoWorldPartitionRuntimeCellTransformer::CanTransfo
 	}
 	else if (!bIsSceneComponentVisible)
 	{
-		return { EFastGeoTransform::Discard, [SceneComponent] { return FString::Printf(TEXT("Component %s has is not visible"), *FastGeo::GetComponentShortName(SceneComponent)); } };
+		return { EFastGeoTransform::Discard, [SceneComponent] { return FString::Printf(TEXT("Component %s is not visible"), *FastGeo::GetComponentShortName(SceneComponent)); } };
 	}
 
 	if (CanIgnoreComponent(InComponent))
 	{
-		return { EFastGeoTransform::Discard };
+		return { EFastGeoTransform::Discard, [InComponent] { return FString::Printf(TEXT("Component %s class is ignored (%s)"), *FastGeo::GetComponentShortName(InComponent), *InComponent->GetClass()->GetName()); } };
 	}
 
 	return EFastGeoTransform::Allow;
