[FastGeo] Strip render-only components and release their exclusive assets on non-rendering processes.
On non-rendering processes (e.g. dedicated server), FastGeo containers previously kept all deserialized component and asset data in memory,  even for render-only content. This change strips that data at load time to reduce memory usage.
- Component stripping:
. Decals and lights are always stripped.
. Primitive components are kept only if they have collision or navigation relevance.
- Asset stripping:
. At cook time, a bitmask is pre-computed identifying which Assets entries are referenced exclusively by render-only components.
. At load time on non-rendering processes, those entries are nulled so GC can reclaim the unreferenced objects (meshes, textures, etc.).
. Material arrays are fully cleared (materials are purely render-related).
. Cluster/HLOD-level asset references are conservatively kept.
- Controlled by FastGeo.StripRenderOnlyComponentsOnNonRenderingProcess (default: true).

#Synced-CL 51624647

diff --git a/Engine/Plugins/Experimental/FastGeoStreaming/Source/FastGeoStreaming/Internal/FastGeoComponentCluster.h b/Engine/Plugins/Experimental/FastGeoStreaming/Source/FastGeoStreaming/Internal/FastGeoComponentCluster.h
--- a/Engine/Plugins/Experimental/FastGeoStreaming/Source/FastGeoStreaming/Internal/FastGeoComponentCluster.h
+++ b/Engine/Plugins/Experimental/FastGeoStreaming/Source/FastGeoStreaming/Internal/FastGeoComponentCluster.h
@@ -53,6 +53,7 @@ public:
 
 #if WITH_EDITOR
 	virtual void PreSave(FObjectPreSaveContext ObjectSaveContext) {}
+	void SerializeWithComponentFilter(FArchive& Ar, TFunctionRef<bool(const FFastGeoComponent&)> ShouldSerializeComponent);
 #endif
 
 	template <typename TComponent = const FFastGeoComponent, typename TFunc>
@@ -97,10 +98,47 @@ protected:
 	TArray<FFastGeoSpotLightComponent> SpotLightComponents;
 	TArray<FFastGeoPointLightComponent> PointLightComponents;
 
+	// Iterates all component arrays. Lambda can return void (non-breakable) or bool (return false to break).
+	template <typename Func>
+	void ForEachComponentArray(Func&& InFunc)
+	{
+		auto Call = [&InFunc](auto& Array) -> bool
+		{
+			if constexpr (std::is_same_v<decltype(InFunc(Array)), bool>)
+			{
+				return InFunc(Array);
+			}
+			else
+			{
+				InFunc(Array);
+				return true;
+			}
+		};
+
+		Call(StaticMeshComponents) &&
+		Call(InstancedStaticMeshComponents) &&
+		Call(SkinnedMeshComponents) &&
+		Call(InstancedSkinnedMeshComponents) &&
+		Call(ProceduralISMComponents) &&
+		Call(DecalComponents) &&
+		Call(RectLightComponents) &&
+		Call(SpotLightComponents) &&
+		Call(PointLightComponents);
+	}
+
+	template <typename Func>
+	void ForEachComponentArray(Func&& InFunc) const
+	{
+		const_cast<FFastGeoComponentCluster*>(this)->ForEachComponentArray([&InFunc](auto& Array) -> decltype(InFunc(std::as_const(Array)))
+		{
+			return InFunc(std::as_const(Array));
+		});
+	}
+
 	template <typename ThisType, typename TComponent, typename TFunc>
 	static void ForEachComponent(ThisType* Self, TFunc&& InFunc)
 	{
-		auto ForEachArray = [&InFunc](auto& InArray)
+		Self->ForEachComponentArray([&InFunc](auto& InArray)
 		{
 			typedef typename TDecay<decltype(InArray)>::Type ArrayType;
 			using TArrayComponent = typename ArrayType::ElementType;
@@ -120,23 +158,14 @@ protected:
 					Forward<TFunc>(InFunc)(*(TComponent*)(&Component));
 				}
 			}
-		};
-		
-		ForEachArray(Self->StaticMeshComponents);
-		ForEachArray(Self->InstancedStaticMeshComponents);
-		ForEachArray(Self->SkinnedMeshComponents);
-		ForEachArray(Self->InstancedSkinnedMeshComponents);
-		ForEachArray(Self->ProceduralISMComponents);
-		ForEachArray(Self->DecalComponents);
-		ForEachArray(Self->RectLightComponents);
-		ForEachArray(Self->SpotLightComponents);
-		ForEachArray(Self->PointLightComponents);
+		});
 	}
 
 	template <typename ThisType, typename TComponent, typename TFunc>
 	static bool ForEachComponentBreakable(ThisType* Self, TFunc&& InFunc)
 	{
-		auto ForEachArray = [&InFunc](auto& InArray) -> bool
+		bool bContinue = true;
+		Self->ForEachComponentArray([&InFunc, &bContinue](auto& InArray) -> bool
 		{
 			typedef typename TDecay<decltype(InArray)>::Type ArrayType;
 			using TArrayComponent = typename ArrayType::ElementType;
@@ -147,6 +176,7 @@ protected:
 				{
 					if (!InFunc(Component))
 					{
+						bContinue = false;
 						return false;
 					}
 				}
@@ -158,26 +188,21 @@ protected:
 					checkSlow(Component.template IsA<TComponent>());
 					if (!Forward<TFunc>(InFunc)(*(TComponent*)(&Component)))
 					{
+						bContinue = false;
 						return false;
 					}
 				}
 			}
-
 			return true;
-		};
-		
-		return ForEachArray(Self->StaticMeshComponents) &&
-			   ForEachArray(Self->InstancedStaticMeshComponents) &&
-			   ForEachArray(Self->SkinnedMeshComponents) &&
-			   ForEachArray(Self->InstancedSkinnedMeshComponents) &&
-			   ForEachArray(Self->ProceduralISMComponents) &&
-			   ForEachArray(Self->DecalComponents) &&
-			   ForEachArray(Self->RectLightComponents) &&
-			   ForEachArray(Self->SpotLightComponents) &&
-			   ForEachArray(Self->PointLightComponents);
+		});
+		return bContinue;
 	}
 
 private:
+
+#if !WITH_EDITOR
+	int32 StripRenderOnlyComponents();
+#endif
 	void ForceUpdateVisibility(const TArray<FFastGeoPrimitiveComponent*>& Components, int32 UpdateCounter);
 	void ForceUpdateVisibilityOrDeferTask(const TArray<FFastGeoPrimitiveComponent*>& Components, int32 UpdateCounter);
 	void UpdateVisibility_Internal(TArray<FFastGeoPrimitiveComponent*>&& ShowComponents, TArray<FFastGeoPrimitiveComponent*>&& HideComponents, int32 UpdateCounter = 0);
diff --git a/Engine/Plugins/Experimental/FastGeoStreaming/Source/FastGeoStreaming/Internal/FastGeoContainer.h b/Engine/Plugins/Experimental/FastGeoStreaming/Source/FastGeoStreaming/Internal/FastGeoContainer.h
--- a/Engine/Plugins/Experimental/FastGeoStreaming/Source/FastGeoStreaming/Internal/FastGeoContainer.h
+++ b/Engine/Plugins/Experimental/FastGeoStreaming/Source/FastGeoStreaming/Internal/FastGeoContainer.h
@@ -10,6 +10,7 @@
 #include "AI/Navigation/NavigationElement.h"
 #include "FastGeoComponentCluster.h"
 #include "FastGeoHLOD.h"
+#include "Containers/BitArray.h"
 #include "UObject/ObjectPtr.h"
 #include "UObject/SoftObjectPtr.h"
 
@@ -304,6 +305,7 @@ private:
 	void ForEachMaterialAsset(TFunctionRef<void(UMaterialInterface*)> Func);
 	void BuildMaterialAssets();
 #if !WITH_EDITOR
+	void ConditionalStripRenderOnlyComponentsAndAssets();
 	void CacheMaterialAssets();
 #endif
 
@@ -325,6 +327,9 @@ private:
 
 	void CollectAssetReferences();
 	void SerializeComponentClusters(FArchive& Ar);
+#if WITH_EDITOR
+	void SerializeComponentClustersWithFilter(FArchive& Ar, TFunctionRef<bool(const FFastGeoComponent&)> ShouldSerializeComponent);
+#endif
 	void RegisterToNavigationSystem();
 	void UnregisterFromNavigationSystem();
 	void OnNavigationInitDone(const UNavigationSystemBase& NavSys);
@@ -448,6 +453,10 @@ private:
 	UPROPERTY()
 	TArray<TSoftObjectPtr<UMaterialInterface>> MaterialAssetNaniteOverrideSoftRefs;
 
+	// Bitmask over Assets: set for entries referenced by non-render-only components.
+	// Computed at cook time. On non-rendering processes, unset entries are nulled at load time.
+	TBitArray<> NonRenderOnlyAssetMask;
+
 	UPROPERTY()
 	TObjectPtr<AFastGeoSurrogateActor> SurrogateActor;
 
@@ -459,6 +468,7 @@ private:
 	UPROPERTY(Transient)
 	TArray<TObjectPtr<UMaterialInterface>> CachedMaterialAssets;
 	bool bMaterialAssetsCached = false;
+	bool bHasStrippedRenderOnlyComponents = false;
 
 	// Transient data
 	TArray<IPrimitiveComponent*> PrimitiveComponents;
diff --git a/Engine/Plugins/Experimental/FastGeoStreaming/Source/FastGeoStreaming/Private/FastGeoComponentCluster.cpp b/Engine/Plugins/Experimental/FastGeoStreaming/Source/FastGeoStreaming/Private/FastGeoComponentCluster.cpp
--- a/Engine/Plugins/Experimental/FastGeoStreaming/Source/FastGeoStreaming/Private/FastGeoComponentCluster.cpp
+++ b/Engine/Plugins/Experimental/FastGeoStreaming/Source/FastGeoStreaming/Private/FastGeoComponentCluster.cpp
@@ -44,15 +44,7 @@ void FFastGeoComponentCluster::Serialize(FArchive& Ar)
 {
 	Ar << Name;
 	Ar << ComponentClusterIndex;
-	Ar << StaticMeshComponents;
-	Ar << InstancedStaticMeshComponents;
-	Ar << SkinnedMeshComponents;
-	Ar << InstancedSkinnedMeshComponents;
-	Ar << ProceduralISMComponents;
-	Ar << DecalComponents;
-	Ar << RectLightComponents;
-	Ar << SpotLightComponents;
-	Ar << PointLightComponents;
+	ForEachComponentArray([&Ar](auto& Array) { Ar << Array; });
 }
 
 void FFastGeoComponentCluster::InitializeDynamicProperties()
@@ -321,13 +313,64 @@ FFastGeoComponent& FFastGeoComponentCluster::AddComponent(FFastGeoElementType In
 
 bool FFastGeoComponentCluster::HasComponents() const
 {
-	return InstancedStaticMeshComponents.Num() > 0 || 
-		   StaticMeshComponents.Num() > 0 || 
-		   InstancedSkinnedMeshComponents.Num() > 0 || 
-		   SkinnedMeshComponents.Num() > 0 || 
-		   ProceduralISMComponents.Num() > 0 || 
-		   DecalComponents.Num() > 0 ||
-		   RectLightComponents.Num() > 0 ||
-		   SpotLightComponents.Num() > 0 ||
-		   PointLightComponents.Num() > 0;
+	bool bHasAny = false;
+	ForEachComponentArray([&bHasAny](const auto& Array) -> bool
+	{
+		bHasAny = Array.Num() > 0;
+		return !bHasAny;
+	});
+	return bHasAny;
+}
+
+#if WITH_EDITOR
+void FFastGeoComponentCluster::SerializeWithComponentFilter(FArchive& Ar, TFunctionRef<bool(const FFastGeoComponent&)> ShouldSerializeComponent)
+{
+	Ar << Name;
+	Ar << ComponentClusterIndex;
+	ForEachComponentArray([&](auto& Array)
+	{
+		for (auto& Component : Array)
+		{
+			if (ShouldSerializeComponent(Component))
+			{
+				Ar << Component;
+			}
+		}
+	});
+}
+#endif
+
+#if !WITH_EDITOR
+int32 FFastGeoComponentCluster::StripRenderOnlyComponents()
+{
+	check(!FApp::CanEverRender());
+
+	int32 NumStripped = 0;
+
+	// Decals and lights are purely visual -- empty unconditionally
+	auto StripAllComponents = [&NumStripped](auto& Array)
+	{
+		NumStripped += Array.Num();
+		Array.Empty();
+	};
+	StripAllComponents(DecalComponents);
+	StripAllComponents(PointLightComponents);
+	StripAllComponents(SpotLightComponents);
+	StripAllComponents(RectLightComponents);
+
+	// Strip remaining components without collision or navigation relevance.
+	// Already-emptied arrays are skipped by RemoveAll (no-op on empty arrays).
+	ForEachComponentArray([&NumStripped](auto& Array)
+	{
+		int32 NumBefore = Array.Num();
+		Array.RemoveAll([](const auto& Component) -> bool
+		{
+			const FFastGeoPrimitiveComponent* Primitive = Component.template CastTo<FFastGeoPrimitiveComponent>();
+			return !Primitive || (!Primitive->IsCollisionEnabled() && !Primitive->IsNavigationRelevant());
+		});
+		NumStripped += NumBefore - Array.Num();
+	});
+
+	return NumStripped;
 }
+#endif
\ No newline at end of file
diff --git a/Engine/Plugins/Experimental/FastGeoStreaming/Source/FastGeoStreaming/Private/FastGeoContainer.cpp b/Engine/Plugins/Experimental/FastGeoStreaming/Source/FastGeoStreaming/Private/FastGeoContainer.cpp
--- a/Engine/Plugins/Experimental/FastGeoStreaming/Source/FastGeoStreaming/Private/FastGeoContainer.cpp
+++ b/Engine/Plugins/Experimental/FastGeoStreaming/Source/FastGeoStreaming/Private/FastGeoContainer.cpp
@@ -47,6 +47,17 @@
 		TEXT("Only taken into account if value is greater than 1."),
 		ECVF_Default);
 
+#if !WITH_EDITOR
+	static bool GStripRenderOnlyComponentsOnNonRenderingProcess = true;
+	static FAutoConsoleVariableRef CVarStripRenderOnlyComponentsOnNonRenderingProcess(
+		TEXT("FastGeo.StripRenderOnlyComponentsOnNonRenderingProcess"),
+		GStripRenderOnlyComponentsOnNonRenderingProcess,
+		TEXT("On non-rendering processes (e.g. dedicated server), strip purely render-only components ")
+		TEXT("(no collision, no navigation) from FastGeo containers to reduce memory usage."),
+		ECVF_ReadOnly);
+#endif
+
+
 	class FAssetRemapArchive : public FArchiveProxy
 	{
 	public:
@@ -1206,6 +1217,80 @@
 	}
 }
 
+class FFastGeoGatherFastGeoContainerAssetRefsArchive : public FArchive
+{
+public:
+	FFastGeoGatherFastGeoContainerAssetRefsArchive(UFastGeoContainer& Container)
+		: FFastGeoGatherFastGeoContainerAssetRefsArchive()
+	{
+		Container.SerializeComponentClusters(*this);
+	}
+
 #if WITH_EDITOR
+	FFastGeoGatherFastGeoContainerAssetRefsArchive(UFastGeoContainer& Container, TFunctionRef<bool(const FFastGeoComponent&)> ShouldSerializeComponent)
+		: FFastGeoGatherFastGeoContainerAssetRefsArchive()
+	{
+		Container.SerializeComponentClustersWithFilter(*this, ShouldSerializeComponent);
+	}
+#endif
+
+	virtual FArchive& operator<<(UObject*& Obj) override
+	{
+		if (Obj && !Obj->IsTemplate() && !Obj->HasAnyFlags(RF_Transient))
+		{
+			if (!UniqueAssets.Contains(Obj))
+			{
+				check(IsKnownClass(Obj->GetClass()));
+				UniqueAssets.Add(Obj);
+			}
+		}
+		return *this;
+	}
+
+	virtual FArchive& operator<<(FObjectPtr& Obj) override
+	{
+		UObject* ObjPtr = Obj.Get();
+		FArchive& Result = operator<<(ObjPtr);
+		Obj = ObjPtr;
+		return Result;
+	}
+
+	const TSet<UObject*>& GetUniqueAssets() { return UniqueAssets; }
+
+private:
+	FFastGeoGatherFastGeoContainerAssetRefsArchive()
+	{
+		KnownClasses.Add(UTexture::StaticClass());
+		KnownClasses.Add(UTextureLightProfile::StaticClass());
+		KnownClasses.Add(URuntimeVirtualTexture::StaticClass());
+		KnownClasses.Add(UStaticMesh::StaticClass());
+		KnownClasses.Add(UMaterialInterface::StaticClass());
+		KnownClasses.Add(UPhysicalMaterial::StaticClass());
+		KnownClasses.Add(USkeletalMesh::StaticClass());
+
+		SetIsPersistent(true);
+		SetIsSaving(true);
+		ArIgnoreOuterRef = true;
+		ArIsObjectReferenceCollector = true;
+		ArShouldSkipBulkData = true;
+	}
+
+	bool IsKnownClass(UClass* Class)
+	{
+		for (UClass* KnownClass : KnownClasses)
+		{
+			if (Class->IsChildOf(KnownClass))
+			{
+				return true;
+			}
+		}
+		return false;
+	}
+
+	TSet<UClass*> KnownClasses;
+	TSet<UObject*> UniqueAssets;
+};
+
+#if WITH_EDITOR
 void UFastGeoContainer::PreSave(FObjectPreSaveContext ObjectSaveContext)
 {
@@ -1260,70 +1345,31 @@
 		// Clear transient cache (will be rebuilt at load time)
 		CachedMaterialAssets.Reset();
 		bMaterialAssetsCached = false;
+
+		// Build non-render-only asset mask for runtime stripping
+		// on non-rendering processes (e.g. dedicated server).
+		// Uses the full serialization path (cluster/HLOD data + components) but filters
+		// out render-only components. This ensures cluster/HLOD-level asset refs are kept.
+		auto IsNonRenderOnly = [](const FFastGeoComponent& Comp) -> bool
+		{
+			const FFastGeoPrimitiveComponent* Primitive = Comp.CastTo<FFastGeoPrimitiveComponent>();
+			return Primitive && (Primitive->IsCollisionEnabled() || Primitive->IsNavigationRelevant());
+		};
+
+		FFastGeoGatherFastGeoContainerAssetRefsArchive NonRenderOnlyGatherAr(*this, IsNonRenderOnly);
+		const TSet<UObject*>& NonRenderOnlyAssetSet = NonRenderOnlyGatherAr.GetUniqueAssets();
+
+		NonRenderOnlyAssetMask.Init(false, Assets.Num());
+		for (int32 i = 0; i < Assets.Num(); ++i)
+		{
+			if (NonRenderOnlyAssetSet.Contains(Assets[i]))
+			{
+				NonRenderOnlyAssetMask[i] = true;
+			}
+		}
 	}
 }
 #endif
 
-class FFastGeoGatherFastGeoContainerAssetRefsArchive : public FArchive
-{
-public:
-	FFastGeoGatherFastGeoContainerAssetRefsArchive(UFastGeoContainer& Container)
-	{
-		KnownClasses.Add(UTexture::StaticClass());
-		KnownClasses.Add(UTextureLightProfile::StaticClass());
-		KnownClasses.Add(URuntimeVirtualTexture::StaticClass());
-		KnownClasses.Add(UStaticMesh::StaticClass());
-		KnownClasses.Add(UMaterialInterface::StaticClass());
-		KnownClasses.Add(UPhysicalMaterial::StaticClass());
-		KnownClasses.Add(USkeletalMesh::StaticClass());
-
-		SetIsPersistent(true);
-		SetIsSaving(true);
-		ArIgnoreOuterRef = true;
-		ArIsObjectReferenceCollector = true;
-		ArShouldSkipBulkData = true;
-		Container.SerializeComponentClusters(*this);
-	}
-
-	virtual FArchive& operator<<(UObject*& Obj) override
-	{
-		if (Obj && !Obj->IsTemplate() && !Obj->HasAnyFlags(RF_Transient))
-		{
-			if (!UniqueAssets.Contains(Obj))
-			{
-				check(IsKnownClass(Obj->GetClass()));
-				UniqueAssets.Add(Obj);
-			}
-		}
-		return *this;
-	}
-
-	virtual FArchive& operator<<(FObjectPtr& Obj) override
-	{
-		UObject* ObjPtr = Obj.Get();
-		FArchive& Result = operator<<(ObjPtr);
-		Obj = ObjPtr;
-		return Result;
-	}
-
-	const TSet<UObject*>& GetUniqueAssets() { return UniqueAssets; }
-
-private:
-	bool IsKnownClass(UClass* Class)
-	{
-		for (UClass* KnownClass : KnownClasses)
-		{
-			if (Class->IsChildOf(KnownClass))
-			{
-				return true;
-			}
-		}
-		return false;
-	}
-
-	TSet<UClass*> KnownClasses;
-	TSet<UObject*> UniqueAssets;
-};
-
 void UFastGeoContainer::CollectAssetReferences()
 {
@@ -1352,10 +1398,15 @@
 
 	FastGeo::Private::FAssetRemapArchive AssetRemapAr(Ar, Assets);
 	SerializeComponentClusters(AssetRemapAr);
+	Ar << NonRenderOnlyAssetMask;
 
 #if !WITH_EDITOR
 	if (Ar.IsLoading())
 	{
+		// On non-rendering processes (e.g. dedicated server), release render-only asset
+		// references and strip render-only components to reduce memory usage.
+		ConditionalStripRenderOnlyComponentsAndAssets();
+
 		// Once loaded, initialize component clusters & components dynamic properties
 		InitializeDynamicProperties(true);
 
@@ -1373,6 +1424,66 @@
 	Ar << HLODs;
 }
 
+#if WITH_EDITOR
+void UFastGeoContainer::SerializeComponentClustersWithFilter(FArchive& Ar, TFunctionRef<bool(const FFastGeoComponent&)> ShouldSerializeComponent)
+{
+	for (auto& Cluster : ComponentClusters)
+	{
+		Cluster.SerializeWithComponentFilter(Ar, ShouldSerializeComponent);
+	}
+	for (auto& HLOD : HLODs)
+	{
+		HLOD.SerializeWithComponentFilter(Ar, ShouldSerializeComponent);
+	}
+}
+#endif
+
+#if !WITH_EDITOR
+void UFastGeoContainer::ConditionalStripRenderOnlyComponentsAndAssets()
+{
+	if (bHasStrippedRenderOnlyComponents)
+	{
+		return;
+	}
+
+	bHasStrippedRenderOnlyComponents = true;
+
+	if (FApp::CanEverRender() || !FastGeo::Private::GStripRenderOnlyComponentsOnNonRenderingProcess)
+	{
+		// Rendering process -- discard the extra cook data, keep all assets
+		NonRenderOnlyAssetMask.Empty();
+		return;
+	}
+
+	TRACE_CPUPROFILER_EVENT_SCOPE(UFastGeoContainer::ConditionalStripRenderOnlyComponentsAndAssets);
+
+	// Null out Assets entries not needed by non-render-only components.
+	// Bits are set for entries to keep.
+	check(NonRenderOnlyAssetMask.Num() == Assets.Num());
+	for (int32 i = 0; i < Assets.Num(); ++i)
+	{
+		if (!NonRenderOnlyAssetMask[i])
+		{
+			Assets[i] = nullptr;
+		}
+	}
+	NonRenderOnlyAssetMask.Empty();
+
+	// Materials are purely render-related -- not needed on non-rendering processes
+	MaterialAssetHardRefs.Empty();
+	MaterialAssetNaniteOverrideSoftRefs.Empty();
+
+	// Strip render-only components
+	int32 TotalStripped = 0;
+	ForEachComponentCluster([&TotalStripped](FFastGeoComponentCluster& ComponentCluster)
+	{
+		TotalStripped += ComponentCluster.StripRenderOnlyComponents();
+	});
+
+	UE_CLOG(TotalStripped > 0, LogFastGeoStreaming, Verbose, TEXT("FastGeo: Stripped %d render-only components in '%s'"), TotalStripped, *GetName());
+}
+#endif
+
 void UFastGeoContainer::OnPrecachePSOsBegin_GameThread()
 {
 #if !WITH_EDITOR && UE_WITH_PSO_PRECACHING
