[FastGeo] Guard PSO precaching against null OwningWorld in OnLevelStreamingStateChanged
When a level transitions between streaming states, OnLevelStreamingStateChanged may fire after the level's OwningWorld has been cleared (e.g. during RemoveFromWorld).
Add a GetWorld() null check before accessing the FastGeo container to prevent potential null dereference through GetWorldSubsystem() -> GetWorld() -> Level->OwningWorld.

#Synced-CL 51624647

diff --git a/Engine/Plugins/Experimental/FastGeoStreaming/Source/FastGeoStreaming/Private/FastGeoWorldSubsystem.cpp b/Engine/Plugins/Experimental/FastGeoStreaming/Source/FastGeoStreaming/Private/FastGeoWorldSubsystem.cpp
--- a/Engine/Plugins/Experimental/FastGeoStreaming/Source/FastGeoStreaming/Private/FastGeoWorldSubsystem.cpp
+++ b/Engine/Plugins/Experimental/FastGeoStreaming/Source/FastGeoStreaming/Private/FastGeoWorldSubsystem.cpp
@@ -251,7 +251,7 @@ void UFastGeoWorldSubsystem::OnLevelStreamingStateChanged(UWorld* World, const U
 		return;
 	}
 
-	if (LevelIfLoaded && ((NewState == ELevelStreamingState::LoadedNotVisible) || (NewState == ELevelStreamingState::LoadedVisible)))
+	if (LevelIfLoaded && LevelIfLoaded->GetWorld() && ((NewState == ELevelStreamingState::LoadedNotVisible) || (NewState == ELevelStreamingState::LoadedVisible)))
 	{
 		if (UFastGeoContainer* FastGeo = LevelIfLoaded->GetAssetUserData<UFastGeoContainer>())
 		{
