[FastGeo] Add support for rejecting actors containing specific component classes
- Add DisallowedActorsContainingComponentClasses and DisallowedActorsContainingExactComponentClasses to UFastGeoTransformerSettings.
- Actors containing a component of a disallowed class (or exact class) are rejected from FastGeo transformation.
- Based on PR #14446 contribution, reworked to fix placement and log messages.

#Synced-CL 51624647

diff --git a/Engine/Plugins/Experimental/FastGeoStreaming/Source/FastGeoStreaming/Private/FastGeoWorldPartitionRuntimeCellTransformer.cpp b/Engine/Plugins/Experimental/FastGeoStreaming/Source/FastGeoStreaming/Private/FastGeoWorldPartitionRuntimeCellTransformer.cpp
--- a/Engine/Plugins/Experimental/FastGeoStreaming/Source/FastGeoStreaming/Private/FastGeoWorldPartitionRuntimeCellTransformer.cpp
+++ b/Engine/Plugins/Experimental/FastGeoStreaming/Source/FastGeoStreaming/Private/FastGeoWorldPartitionRuntimeCellTransformer.cpp
@@ -80,30 +80,30 @@ namespace FastGeo
 		UE_LOG(LogFastGeoStreaming, Log, TEXT("%s"), LogSeparatorOuter);
 	}
 
-	static IWorldPartitionHLODObject* GetHLODObject(AActor* Actor)
+	static const IWorldPartitionHLODObject* GetHLODObject(const AActor* Actor)
 	{
-		if (AWorldPartitionHLOD* HLODActor = Cast<AWorldPartitionHLOD>(Actor))
+		if (const AWorldPartitionHLOD* HLODActor = Cast<AWorldPartitionHLOD>(Actor))
 		{
 			return HLODActor;
 		}
-		if (AWorldPartitionCustomHLOD* CustomHLODActor = Cast<AWorldPartitionCustomHLOD>(Actor))
+		if (const AWorldPartitionCustomHLOD* CustomHLODActor = Cast<AWorldPartitionCustomHLOD>(Actor))
 		{
 			return CustomHLODActor;
 		}
 		return nullptr;
 	}
 
-	static bool IsHLODActor(AActor* Actor)
+	static bool IsHLODActor(const AActor* Actor)
 	{
 		return GetHLODObject(Actor) != nullptr;
 	}
 
-	static bool IsCollisionEnabled(UPrimitiveComponent* Component)
+	static bool IsCollisionEnabled(const UPrimitiveComponent* Component)
 	{
 		return Component->IsCollisionEnabled() && !IsHLODActor(Component->GetOwner());
 	}
 
-	static FString GetComponentShortName(UActorComponent* InComponent)
+	static FString GetComponentShortName(const UActorComponent* InComponent)
 	{
 		TStringBuilder<256> Builder;
 		Builder += InComponent->GetOwner()->GetName();
@@ -384,7 +384,7 @@ void UFastGeoWorldPartitionRuntimeCellTransformer::Transform(ULevel* InLevel)
 
 			FFastGeoComponentCluster* CurrentComponentCluster = LevelComponentCluster.Get();
 			TUniquePtr<FFastGeoHLOD> FastGeoHLOD;
-			IWorldPartitionHLODObject* HLODActor = FastGeo::GetHLODObject(Actor);
+			const IWorldPartitionHLODObject* HLODActor = FastGeo::GetHLODObject(Actor);
 			if (HLODActor)
 			{
 				FastGeoHLOD.Reset(new FFastGeoHLOD(FastGeo, *FString::Printf(TEXT("FastGeoHLOD_%s"), *Actor->GetName())));
@@ -1023,6 +1023,34 @@ FFastGeoTransformResult UFastGeoWorldPartitionRuntimeCellTransformer::CanTransfo
 		return AllowedActorClassResult;
 	}
 
+	// Reject actors containing disallowed component classes
+	{
+		const UFastGeoTransformerSettings& EffectiveSettings = GetSettings();
+		TArray<UActorComponent*> ActorComponents;
+		InActor->GetComponents(ActorComponents);
+
+		for (const UActorComponent* ActorComponent : ActorComponents)
+		{
+			const UClass* ComponentClass = ActorComponent->GetClass();
+
+			for (TSubclassOf<UActorComponent> DisallowedComponentClass : EffectiveSettings.DisallowedActorsContainingComponentClasses)
+			{
+				if (ComponentClass->IsChildOf(DisallowedComponentClass))
+				{
+					return { EFastGeoTransform::Reject, [InActor, ActorComponent, DisallowedComponentClass] { return FString::Printf(TEXT("Actor %s contains a component %s of disallowed class (%s)"), *InActor->GetName(), *FastGeo::GetComponentShortName(ActorComponent), *DisallowedComponentClass->GetName()); } };
+				}
+			}
+
+			for (TSubclassOf<UActorComponent> DisallowedExactComponentClass : EffectiveSettings.DisallowedActorsContainingExactComponentClasses)
+			{
+				if (ComponentClass == DisallowedExactComponentClass)
+				{
+					return { EFastGeoTransform::Reject, [InActor, ActorComponent, DisallowedExactComponentClass] { return FString::Printf(TEXT("Actor %s contains a component %s of disallowed exact class (%s)"), *InActor->GetName(), *FastGeo::GetComponentShortName(ActorComponent), *DisallowedExactComponentClass->GetName()); } };
+				}
+			}
+		}
+	}
+
 	FString Reason;
 
 	if (InActor->ActorHasTag(NAME_CellTransformerIgnoreActor))
diff --git a/Engine/Plugins/Experimental/FastGeoStreaming/Source/FastGeoStreaming/Public/FastGeoWorldPartitionRuntimeCellTransformer.h b/Engine/Plugins/Experimental/FastGeoStreaming/Source/FastGeoStreaming/Public/FastGeoWorldPartitionRuntimeCellTransformer.h
--- a/Engine/Plugins/Experimental/FastGeoStreaming/Source/FastGeoStreaming/Public/FastGeoWorldPartitionRuntimeCellTransformer.h
+++ b/Engine/Plugins/Experimental/FastGeoStreaming/Source/FastGeoStreaming/Public/FastGeoWorldPartitionRuntimeCellTransformer.h
@@ -133,6 +133,14 @@ public:
 	UPROPERTY(EditAnywhere, Category="Settings", meta = (AllowAbstract = "true"))
 	TArray<TSubclassOf<UActorComponent>> DisallowedExactComponentClasses;
 
+	/** Disallow any actor that contains a component of these classes (recursive) */
+	UPROPERTY(EditAnywhere, Category="Settings", meta = (AllowAbstract = "true"))
+	TArray<TSubclassOf<UActorComponent>> DisallowedActorsContainingComponentClasses;
+
+	/** Disallow any actor that contains a component of these exact classes */
+	UPROPERTY(EditAnywhere, Category="Settings", meta = (AllowAbstract = "true"))
+	TArray<TSubclassOf<UActorComponent>> DisallowedActorsContainingExactComponentClasses;
+
 	/** Component classes (recursive) to ignore when deciding to destroy a converted actor to FastGeo */
 	UPROPERTY(EditAnywhere, Category="Settings", meta = (AllowAbstract = "true"))
 	TArray<TSubclassOf<UActorComponent>> IgnoredRemainingComponentClasses;
