[FastGeo][PCG] Editor hit proxy support for procedural ISM components

Clicking a FastGeo procedural ISM component in the editor viewport now selects the generating PCG actor instead of clicking through to geometry behind it. The selection actor will either be the partition actor if actor-componentless generation is disabled, otherwise the original PCG actor.

#rnx
#Synced-CL 51942803

diff --git a/Engine/Plugins/Experimental/FastGeoStreaming/Source/FastGeoStreaming/Internal/FastGeoPrimitiveComponent.h b/Engine/Plugins/Experimental/FastGeoStreaming/Source/FastGeoStreaming/Internal/FastGeoPrimitiveComponent.h
--- a/Engine/Plugins/Experimental/FastGeoStreaming/Source/FastGeoStreaming/Internal/FastGeoPrimitiveComponent.h
+++ b/Engine/Plugins/Experimental/FastGeoStreaming/Source/FastGeoStreaming/Internal/FastGeoPrimitiveComponent.h
@@ -152,7 +152,7 @@ protected:
 #endif
 	virtual void InitializeSceneProxyDescDynamicProperties();
 	virtual FPrimitiveSceneProxy* CreateSceneProxy(ESceneProxyCreationError* OutError) = 0;
-	FPrimitiveSceneDesc BuildSceneDesc();
+	virtual FPrimitiveSceneDesc BuildSceneDesc();
 
 	//~ Begin Physics
 	virtual bool IsStaticPhysics() const;
diff --git a/Engine/Plugins/Experimental/FastGeoStreaming/Source/FastGeoStreaming/Internal/FastGeoProceduralISMComponent.h b/Engine/Plugins/Experimental/FastGeoStreaming/Source/FastGeoStreaming/Internal/FastGeoProceduralISMComponent.h
--- a/Engine/Plugins/Experimental/FastGeoStreaming/Source/FastGeoStreaming/Internal/FastGeoProceduralISMComponent.h
+++ b/Engine/Plugins/Experimental/FastGeoStreaming/Source/FastGeoStreaming/Internal/FastGeoProceduralISMComponent.h
@@ -4,6 +4,9 @@
 
 #include "FastGeoStaticMeshComponent.h"
 #include "InstancedStaticMeshSceneProxyDesc.h"
+#include "UObject/WeakObjectPtrTemplates.h"
+
+class AActor;
 
 struct FProceduralISMComponentDescriptor;
 
@@ -23,6 +26,16 @@ public:
 	/** Set spatial hashes which provide bounds for sub-ranges of instances within this primitive, enabling early CPU culling. */
 	FASTGEOSTREAMING_API void SetSpatialHashes(TArray<FInstanceSceneDataBuffers::FCompressedSpatialHashItem>&& InSpatialHashes);
 
+#if WITH_EDITOR
+	/** Sets an actor to use as the hit proxy target when this component is clicked in the editor.
+	 *  If set, clicking this primitive in the editor will select InActor instead of the container.
+	 *  Must be called before Register().
+	 */
+	FASTGEOSTREAMING_API void SetHitProxyTargetActor(AActor* InActor);
+
+	FASTGEOSTREAMING_API void UpdateEditorSelectionState();
+#endif
+
 	//~ Begin FFastGeoComponent Interface
 	virtual void Serialize(FArchive& Ar) override;
 	//~ End FFastGeoComponent Interface
@@ -31,12 +44,15 @@ protected:
 	//~ Begin FFastGeoPrimitiveComponent interface
 	virtual FPrimitiveSceneProxyDesc& GetSceneProxyDesc() override { return SceneProxyDesc; }
 	virtual const FPrimitiveSceneProxyDesc& GetSceneProxyDesc() const override { return SceneProxyDesc; }
+	virtual FPrimitiveSceneDesc BuildSceneDesc() override;
 	virtual bool IsNavigationRelevant() const override { return false; }
 #if WITH_EDITOR
 	virtual void ResetSceneProxyDescUnsupportedProperties() override;
+	virtual HHitProxy* CreateMeshHitProxy(int32 SectionIndex, int32 MaterialIndex) override;
 #endif
 	virtual void InitializeSceneProxyDescDynamicProperties() override;
 	virtual void ApplyWorldTransform(const FTransform& InTransform) override;
+	virtual HHitProxy* CreatePrimitiveHitProxies(TArray<TRefCountPtr<HHitProxy>>& OutHitProxies) override;
 	//~ End FFastGeoPrimitiveComponent interface
 
 	//~ Begin FFastGeoStaticMeshComponentBase interface
@@ -60,6 +76,9 @@ private:
 	// Transient data
 	TSharedPtr<FInstanceDataSceneProxy, ESPMode::ThreadSafe> DataProxy{};
 	TArray<FInstanceSceneDataBuffers::FCompressedSpatialHashItem> SpatialHashes;
+#if WITH_EDITOR
+	TWeakObjectPtr<AActor> HitProxyTargetActor;
+#endif
 
 	friend class FInstancedStaticMeshComponentHelper;
 };
diff --git a/Engine/Plugins/Experimental/FastGeoStreaming/Source/FastGeoStreaming/Private/FastGeoProceduralISMComponent.cpp b/Engine/Plugins/Experimental/FastGeoStreaming/Source/FastGeoStreaming/Private/FastGeoProceduralISMComponent.cpp
--- a/Engine/Plugins/Experimental/FastGeoStreaming/Source/FastGeoStreaming/Private/FastGeoProceduralISMComponent.cpp
+++ b/Engine/Plugins/Experimental/FastGeoStreaming/Source/FastGeoStreaming/Private/FastGeoProceduralISMComponent.cpp
@@ -4,8 +4,12 @@
 
 #include "InstancedStaticMeshComponentHelper.h"
 #include "NaniteSceneProxy.h"
+#include "PrimitiveSceneDesc.h"
 #include "InstancedStaticMesh/ISMInstanceDataSceneProxy.h"
 #include "ISMPartition/ProceduralISMComponentDescriptor.h"
+#if WITH_EDITOR
+#include "EngineUtils.h"
+#endif
 
 const FFastGeoElementType FFastGeoProceduralISMComponent::Type(&FFastGeoStaticMeshComponentBase::Type);
 
@@ -83,7 +87,43 @@ void FFastGeoProceduralISMComponent::Serialize(FArchive& Ar)
 	FArchive_Serialize_BitfieldBool(Ar, SceneProxyDesc.bUseGpuLodSelection);
 }
 
+FPrimitiveSceneDesc FFastGeoProceduralISMComponent::BuildSceneDesc()
+{
+	FPrimitiveSceneDesc SceneDesc = Super::BuildSceneDesc();
 #if WITH_EDITOR
+	SceneDesc.PrimitiveComponentInterface = this;
+#endif
+
+	return SceneDesc;
+}
+
+#if WITH_EDITOR
+void FFastGeoProceduralISMComponent::SetHitProxyTargetActor(AActor* InActor)
+{
+	checkf(!IsRegistered(), TEXT("SetHitProxyTargetActor must be called before Register()"));
+	HitProxyTargetActor = InActor;
+}
+
+void FFastGeoProceduralISMComponent::UpdateEditorSelectionState()
+{
+	if (AActor* Actor = HitProxyTargetActor.Get())
+	{
+		if (FPrimitiveSceneProxy* Proxy = GetSceneProxy())
+		{
+			Proxy->SetSelection_GameThread(Actor->IsSelected(), /*bInIndividuallySelected=*/false);
+		}
+	}
+}
+
+HHitProxy* FFastGeoProceduralISMComponent::CreateMeshHitProxy(int32 SectionIndex, int32 MaterialIndex)
+{
+	if (AActor* Actor = HitProxyTargetActor.Get())
+	{
+		return new HActor(Actor, /*InPrimComponent=*/nullptr);
+	}
+	return nullptr;
+}
+
 void FFastGeoProceduralISMComponent::ResetSceneProxyDescUnsupportedProperties()
 {
 	Super::ResetSceneProxyDescUnsupportedProperties();
@@ -96,6 +136,19 @@ void FFastGeoProceduralISMComponent::ResetSceneProxyDescUnsupportedProperties()
 }
 #endif
 
+HHitProxy* FFastGeoProceduralISMComponent::CreatePrimitiveHitProxies(TArray<TRefCountPtr<HHitProxy>>& OutHitProxies)
+{
+#if WITH_EDITOR
+	if (AActor* Actor = HitProxyTargetActor.Get())
+	{
+		HHitProxy* Proxy = new HActor(Actor, /*InPrimComponent=*/nullptr);
+		OutHitProxies.Add(Proxy);
+		return Proxy;
+	}
+#endif
+	return nullptr;
+}
+
 void FFastGeoProceduralISMComponent::InitializeSceneProxyDescDynamicProperties()
 {
 	Super::InitializeSceneProxyDescDynamicProperties();
@@ -103,6 +156,17 @@ void FFastGeoProceduralISMComponent::InitializeSceneProxyDescDynamicProperties()
 	SceneProxyDesc.bAffectDynamicIndirectLighting = false;
 	SceneProxyDesc.bAffectDistanceFieldLighting = false;
 	SceneProxyDesc.bCollisionEnabled = false;
+
+#if WITH_EDITOR
+	// ResetSceneProxyDescUnsupportedProperties (called from Super above) marks bSelectable=false for all FastGeo
+	// primitives. Re-enable it if a hit proxy target actor was set so that this primitive is rendered into
+	// the hit proxy buffer and can be clicked in the editor.
+	if (HitProxyTargetActor.IsValid())
+	{
+		SceneProxyDesc.bSelectable = true;
+		SceneProxyDesc.bSelected = HitProxyTargetActor->IsSelected();
+	}
+#endif
 }
 
 void FFastGeoProceduralISMComponent::ApplyWorldTransform(const FTransform& InTransform)
diff --git a/Engine/Plugins/Experimental/FastGeoStreaming/Source/FastGeoStreamingEditor/Private/FastGeoStreamingEditorModule.cpp b/Engine/Plugins/Experimental/FastGeoStreaming/Source/FastGeoStreamingEditor/Private/FastGeoStreamingEditorModule.cpp
--- a/Engine/Plugins/Experimental/FastGeoStreaming/Source/FastGeoStreamingEditor/Private/FastGeoStreamingEditorModule.cpp
+++ b/Engine/Plugins/Experimental/FastGeoStreaming/Source/FastGeoStreamingEditor/Private/FastGeoStreamingEditorModule.cpp
@@ -5,6 +5,11 @@
 #include "FastGeoWorldPartitionRuntimeCellTransformer.h"
 #include "PropertyEditorModule.h"
 #include "Modules/ModuleManager.h"
+#include "FastGeoComponentCluster.h"
+#include "FastGeoContainer.h"
+#include "FastGeoProceduralISMComponent.h"
+#include "Selection.h"
+#include "UObject/UObjectIterator.h"
 
 IMPLEMENT_MODULE(FFastGeoStreamingEditorModule, FastGeoStreamingEditor);
 
@@ -15,6 +20,8 @@ void FFastGeoStreamingEditorModule::StartupModule()
 	{
 		PropertyModule->RegisterCustomPropertyTypeLayout(FFastGeoConversionButton::StaticStruct()->GetFName(), FOnGetPropertyTypeCustomizationInstance::CreateStatic(&FConvertFastGeoSettingsAssetButtonCustomization::MakeInstance));
 	}
+
+	USelection::SelectionChangedEvent.AddRaw(this, &FFastGeoStreamingEditorModule::OnEditorSelectionChanged);
 }
 
 void FFastGeoStreamingEditorModule::ShutdownModule()
@@ -26,4 +33,21 @@ void FFastGeoStreamingEditorModule::ShutdownModule()
 			PropertyModule->UnregisterCustomPropertyTypeLayout(FFastGeoConversionButton::StaticStruct()->GetFName());
 		}
 	}
-}
\ No newline at end of file
+
+	USelection::SelectionChangedEvent.RemoveAll(this);
+}
+
+void FFastGeoStreamingEditorModule::OnEditorSelectionChanged(UObject* /*ChangedObject*/)
+{
+	for (TObjectIterator<UFastGeoContainer> It(RF_ClassDefaultObject, /*bIncludeDerivedClasses=*/true, EInternalObjectFlags::Garbage); It; ++It)
+	{
+		It->ForEachComponentCluster([](FFastGeoComponentCluster& Cluster)
+		{
+			// Currently only Procedural ISM components support selection.
+			Cluster.ForEachComponent<FFastGeoProceduralISMComponent>([](FFastGeoProceduralISMComponent& Component)
+			{
+				Component.UpdateEditorSelectionState();
+			});
+		});
+	}
+}
diff --git a/Engine/Plugins/Experimental/FastGeoStreaming/Source/FastGeoStreamingEditor/Private/FastGeoStreamingEditorModule.h b/Engine/Plugins/Experimental/FastGeoStreaming/Source/FastGeoStreamingEditor/Private/FastGeoStreamingEditorModule.h
--- a/Engine/Plugins/Experimental/FastGeoStreaming/Source/FastGeoStreamingEditor/Private/FastGeoStreamingEditorModule.h
+++ b/Engine/Plugins/Experimental/FastGeoStreaming/Source/FastGeoStreamingEditor/Private/FastGeoStreamingEditorModule.h
@@ -10,4 +10,7 @@ class FFastGeoStreamingEditorModule : public IModuleInterface
 public:
 	virtual void StartupModule() override;
 	virtual void ShutdownModule() override;
+
+private:
+	void OnEditorSelectionChanged(UObject* /*ChangedObject*/);
 };
\ No newline at end of file
diff --git a/Engine/Plugins/Experimental/PCGInterops/PCGFastGeoInterop/Source/PCGFastGeoInterop/Private/Compute/PrimitiveFactories/PCGPrimitiveFactoryFastGeoPISMC.cpp b/Engine/Plugins/Experimental/PCGInterops/PCGFastGeoInterop/Source/PCGFastGeoInterop/Private/Compute/PrimitiveFactories/PCGPrimitiveFactoryFastGeoPISMC.cpp
--- a/Engine/Plugins/Experimental/PCGInterops/PCGFastGeoInterop/Source/PCGFastGeoInterop/Private/Compute/PrimitiveFactories/PCGPrimitiveFactoryFastGeoPISMC.cpp
+++ b/Engine/Plugins/Experimental/PCGInterops/PCGFastGeoInterop/Source/PCGFastGeoInterop/Private/Compute/PrimitiveFactories/PCGPrimitiveFactoryFastGeoPISMC.cpp
@@ -115,7 +115,22 @@ bool FPCGPrimitiveFactoryFastGeoPISMC::Create(FPCGContext* InContext)
 		return true;
 	}
 
-	UWorld* World = InContext->ExecutionSource->GetExecutionState().GetWorld();
+	UWorld* World = ExecutionSource->GetExecutionState().GetWorld();
+
+#if WITH_EDITOR
+	AActor* HitProxyTargetActor = ExecutionSource->GetExecutionState().GetTypedTarget<AActor>();
+
+	// For actor-component-less generation the local execution source has no target actor.
+	// The original source (e.g. the UPCGComponent on the user's actor) typically does.
+	if (!HitProxyTargetActor)
+	{
+		if (IPCGGraphExecutionSource* OriginalSource = ExecutionSource->GetExecutionState().GetOriginalSource())
+		{
+			HitProxyTargetActor = OriginalSource->GetExecutionState().GetTypedTarget<AActor>();
+		}
+	}
+#endif // WITH_EDITOR
+
 	UFastGeoContainer* FastGeo = UFastGeoContainer::CreateTransient(World);
 	if (!FastGeo)
 	{
@@ -149,6 +164,9 @@ bool FPCGPrimitiveFactoryFastGeoPISMC::Create(FPCGContext* InContext)
 		FastGeoComponent.InitializeFromComponentDescriptor(Desc);
 		FastGeoComponent.SetCustomPrimitiveData(CustomPrimitiveData.GetViewForIndex(DescIndex));
 		FastGeoComponent.SetSpatialHashes(MoveTemp(SpatialHashes));
+#if WITH_EDITOR
+		FastGeoComponent.SetHitProxyTargetActor(HitProxyTargetActor);
+#endif
 
 		MeshBounds[DescIndex] = Desc.StaticMesh->GetBoundingBox();
 	}
diff --git a/Engine/Source/Runtime/Engine/Private/EngineUtils.cpp b/Engine/Source/Runtime/Engine/Private/EngineUtils.cpp
--- a/Engine/Source/Runtime/Engine/Private/EngineUtils.cpp
+++ b/Engine/Source/Runtime/Engine/Private/EngineUtils.cpp
@@ -67,7 +67,7 @@ FTypedElementHandle HActor::GetElementHandle() const
 bool HActor::AlwaysAllowsTranslucentPrimitives() const
 {
 #if WITH_EDITOR
-	return PrimComponent->bAlwaysAllowTranslucentSelect;
+	return PrimComponent && PrimComponent->bAlwaysAllowTranslucentSelect;
 #else
 	return false;
 #endif
