[FastGeo] Skip BodyInstance serialization for FastGeo components with collision disabled (mostly HLODs) to save cooking/loading time.
Also zero-init FBodyInstance before deserialization on load to ensure non-UPROPERTY fields match UObject memzero behavior (pre-existing issue for all FastGeo components).

#Synced-CL 51624647

diff --git a/Engine/Plugins/Experimental/FastGeoStreaming/Source/FastGeoStreaming/Private/FastGeoPrimitiveComponent.cpp b/Engine/Plugins/Experimental/FastGeoStreaming/Source/FastGeoStreaming/Private/FastGeoPrimitiveComponent.cpp
--- a/Engine/Plugins/Experimental/FastGeoStreaming/Source/FastGeoStreaming/Private/FastGeoPrimitiveComponent.cpp
+++ b/Engine/Plugins/Experimental/FastGeoStreaming/Source/FastGeoStreaming/Private/FastGeoPrimitiveComponent.cpp
@@ -79,13 +79,6 @@ void FFastGeoPrimitiveComponent::Serialize(FArchive& Ar)
 	Ar << CustomPrimitiveData.Data;
 	Ar << bHasCustomNavigableGeometry;
 	Ar << RuntimeVirtualTextures;
-	FBodyInstance::StaticStruct()->SerializeItem(Ar, &BodyInstance, nullptr);
-	Ar << BodyInstance.InstanceBodyIndex;
-
-	if (Ar.IsLoading() && !UFastGeoWorldSubsystem::ShouldAllowSurrogateComponents())
-	{
-		BodyInstance.InstanceBodyIndex = INDEX_NONE;
-	}
 
 	// Serialize persistent data from FPrimitiveSceneProxyDesc
 	FPrimitiveSceneProxyDesc& SceneProxyDesc = GetSceneProxyDesc();
@@ -160,6 +153,27 @@ void FFastGeoPrimitiveComponent::Serialize(FArchive& Ar)
 	Ar << SceneProxyDesc.RayTracingGroupId;
 	Ar << SceneProxyDesc.VirtualTextureRenderPassType;
 	Ar << SceneProxyDesc.VirtualTextureMainPassMaxDrawDistance;
+
+	// Zero-init + reconstruct BodyInstance to match UObject memzero behavior
+	// for non-UPROPERTY fields that SerializeItem won't restore.
+	if (Ar.IsLoading())
+	{
+		BodyInstance.~FBodyInstance();
+		FMemory::Memzero(&BodyInstance, sizeof(FBodyInstance));
+		new (&BodyInstance) FBodyInstance();
+	}
+
+	// Skip BodyInstance serialization when collision is disabled (e.g. HLODs)
+	if (SceneProxyDesc.bCollisionEnabled)
+	{
+		FBodyInstance::StaticStruct()->SerializeItem(Ar, &BodyInstance, nullptr);
+		Ar << BodyInstance.InstanceBodyIndex;
+	}
+
+	if (Ar.IsLoading() && !UFastGeoWorldSubsystem::ShouldAllowSurrogateComponents())
+	{
+		BodyInstance.InstanceBodyIndex = INDEX_NONE;
+	}
 }
 
 #if WITH_EDITOR
