[FastGeo] UFastGeoSurrogateComponent::IsPhysicsStateCreated now accurately reflects whether body instances have been registered.
Physics state is created through the normal engine lifecycle (CreatePhysicsState/DestroyPhysicsState) when body instances are registered/unregistered, rather than during component registration when no body instances exist yet.
Also overrides OnDestroyPhysicsState to skip UPrimitiveComponent destroy logic since the surrogate never creates its own physics bodies.

#Synced-CL 51624647

diff --git a/Engine/Plugins/Experimental/FastGeoStreaming/Source/FastGeoStreaming/Internal/FastGeoSurrogateComponent.h b/Engine/Plugins/Experimental/FastGeoStreaming/Source/FastGeoStreaming/Internal/FastGeoSurrogateComponent.h
--- a/Engine/Plugins/Experimental/FastGeoStreaming/Source/FastGeoStreaming/Internal/FastGeoSurrogateComponent.h
+++ b/Engine/Plugins/Experimental/FastGeoStreaming/Source/FastGeoStreaming/Internal/FastGeoSurrogateComponent.h
@@ -22,6 +22,7 @@ public:
 	virtual void OnRegister() override;
 	virtual bool ShouldCreatePhysicsState() const override;
 	virtual void OnCreatePhysicsState() override;
+	virtual void OnDestroyPhysicsState() override;
 	//~ End UActorComponent interface
 
 	//~ Begin UPrimitiveComponent interface
@@ -46,6 +47,7 @@ public:
 
 protected:
 
+	bool HasRegisteredBodyInstances() const { return !IndexToBodyInstance.IsEmpty(); }
 	void SetSurrogateComponentDescriptor(const FFastGeoSurrogateComponentDescriptor& InComponentDesc);
 	int32 GetDescriptorIndex() const { return DescriptorIndex; }
 
diff --git a/Engine/Plugins/Experimental/FastGeoStreaming/Source/FastGeoStreaming/Private/FastGeoSurrogateComponent.cpp b/Engine/Plugins/Experimental/FastGeoStreaming/Source/FastGeoStreaming/Private/FastGeoSurrogateComponent.cpp
--- a/Engine/Plugins/Experimental/FastGeoStreaming/Source/FastGeoStreaming/Private/FastGeoSurrogateComponent.cpp
+++ b/Engine/Plugins/Experimental/FastGeoStreaming/Source/FastGeoStreaming/Private/FastGeoSurrogateComponent.cpp
@@ -69,18 +69,29 @@ bool UFastGeoSurrogateComponent::ShouldCreatePhysicsState() const
 	{
 		return false;
 	}
+
+	if (!HasRegisteredBodyInstances())
+	{
+		return false;
+	}
+
 	return Super::ShouldCreatePhysicsState();
 }
 
 void UFastGeoSurrogateComponent::OnCreatePhysicsState()
 {
-	// Skip PrimitiveComponent implementation
+	// Skip UPrimitiveComponent implementation - the surrogate does not create its own physics bodies
 	USceneComponent::OnCreatePhysicsState();
 
-	if (!IsAsyncCreatePhysicsStateRunning())
-	{
-		OnComponentPhysicsStateChanged.Broadcast(this, EComponentPhysicsStateChange::Created);
-	}
+	OnComponentPhysicsStateChanged.Broadcast(this, EComponentPhysicsStateChange::Created);
+}
+
+void UFastGeoSurrogateComponent::OnDestroyPhysicsState()
+{
+	// Skip UPrimitiveComponent implementation - the surrogate does not create its own physics bodies
+	USceneComponent::OnDestroyPhysicsState();
+
+	OnComponentPhysicsStateChanged.Broadcast(this, EComponentPhysicsStateChange::Destroyed);
 }
 
 ECollisionEnabled::Type UFastGeoSurrogateComponent::GetCollisionEnabled() const
@@ -131,6 +142,7 @@ void UFastGeoSurrogateComponent::RegisterBodyInstances(const TArray<FBodyInstanc
 	TRACE_CPUPROFILER_EVENT_SCOPE(UFastGeoSurrogateComponent::RegisterBodyInstances);
 	check(RegisteredBodyInstanceCount == 0);
 	check(IndexToBodyInstance.IsEmpty());
+	check(!IsPhysicsStateCreated());
 	IndexToBodyInstance.SetNumZeroed(InBodyInstances.Num());
 	for (FBodyInstance* BodyInst : InBodyInstances)
 	{
@@ -143,11 +155,26 @@ void UFastGeoSurrogateComponent::RegisterBodyInstances(const TArray<FBodyInstanc
 			IndexToBodyInstance[ArrayIndex] = BodyInst;
 		}
 	}
+
+	// Body instances are now accessible -- create physics state through the
+	// normal engine lifecycle. ShouldCreatePhysicsState gates on
+	// HasRegisteredBodyInstances, so this only succeeds after population.
+	CreatePhysicsState();
 }
 
 void UFastGeoSurrogateComponent::UnregisterBodyInstances()
 {
 	TRACE_CPUPROFILER_EVENT_SCOPE(UFastGeoSurrogateComponent::UnregisterBodyInstances);
+
+	// Destroy physics state before clearing body instances so that
+	// IsPhysicsStateCreated accurately reflects that the surrogate no
+	// longer has valid body instances. Goes through the normal engine
+	// lifecycle (delegates, OnDestroyPhysicsState, etc.).
+	if (IsPhysicsStateCreated())
+	{
+		DestroyPhysicsState();
+	}
+
 	RegisteredBodyInstanceCount = 0;
 	IndexToBodyInstance.Reset();
 }
