[FastGeo] Track async PSO precaching as pending work and launch immediately.
- Set PendingCreate.PrecachePSOState in OnPrecachePSOsBegin_GameThread so HasAnyPendingTasks returns true during async PSO work.
- Implement OnPrecachePSOsEnd_GameThread to reset PrecachePSOState on completion.
- Call TickRenderWork after StartPrecachePSOWork in PrecachePSOs to launch the job immediately instead of waiting for a future container tick.

#Synced-CL 51624647

diff --git a/Engine/Plugins/Experimental/FastGeoStreaming/Source/FastGeoStreaming/Internal/FastGeoContainer.h b/Engine/Plugins/Experimental/FastGeoStreaming/Source/FastGeoStreaming/Internal/FastGeoContainer.h
--- a/Engine/Plugins/Experimental/FastGeoStreaming/Source/FastGeoStreaming/Internal/FastGeoContainer.h
+++ b/Engine/Plugins/Experimental/FastGeoStreaming/Source/FastGeoStreaming/Internal/FastGeoContainer.h
@@ -165,11 +165,14 @@ struct FFastGeoComponentClusterPendingState
 		bIsInBlockingWait = false;
 	}
 
-	bool HasAnyPendingState() const
-	{ 
+	bool HasAnyPendingState(bool bIncludePSOPrecaching = true) const
+	{
 		bool bHasAnyPendingState = !RenderState.IsCompleted() || !PhysicsState.IsCompleted();
 #if !WITH_EDITOR && UE_WITH_PSO_PRECACHING
-		bHasAnyPendingState |= !PrecachePSOState.IsCompleted();
+		if (bIncludePSOPrecaching)
+		{
+			bHasAnyPendingState |= !PrecachePSOState.IsCompleted();
+		}
 #endif
 		return bHasAnyPendingState;
 	}
@@ -201,8 +204,8 @@ public:
 	bool IsFullyRegistered() const;
 	bool IsFullyUnregistered() const;
 	void Tick(bool bWaitForCompletion = false);
-	bool HasAnyPendingTasks() const;
-	bool HasAnyPendingCreateTasks() const;
+	bool HasAnyPendingTasks(bool bIncludePSOPrecaching = true) const;
+	bool HasAnyPendingCreateTasks(bool bIncludePSOPrecaching = true) const;
 	bool HasAnyPendingDestroyTasks() const;
 	void PrecachePSOs();
 	UWorld* GetWorld() const;
@@ -235,7 +238,7 @@ public:
 
 	void OnPrecachePSOsBegin_GameThread();
 	void OnPrecachePSOs_Concurrent();
-	void OnPrecachePSOsEnd_GameThread() {}
+	void OnPrecachePSOsEnd_GameThread();
 
 	void OnCreateRenderStateBegin_GameThread();
 	void OnCreateRenderState_Concurrent();
diff --git a/Engine/Plugins/Experimental/FastGeoStreaming/Source/FastGeoStreaming/Private/FastGeoContainer.cpp b/Engine/Plugins/Experimental/FastGeoStreaming/Source/FastGeoStreaming/Private/FastGeoContainer.cpp
--- a/Engine/Plugins/Experimental/FastGeoStreaming/Source/FastGeoStreaming/Private/FastGeoContainer.cpp
+++ b/Engine/Plugins/Experimental/FastGeoStreaming/Source/FastGeoStreaming/Private/FastGeoContainer.cpp
@@ -256,7 +256,8 @@ void UFastGeoContainer::StartRegisterTransition()
 	check(RegistrationState == ERegistrationState::Unregistered);
 	check(!IsUsingSurrogateComponents() || !SurrogateActor->HasRegisteredWithFastGeo());
 
-	if (const bool bHasAnyPendingTasks = HasAnyPendingTasks())
+	constexpr bool bIncludePSOPrecaching = false; // PSO precaching is allowed to run before registration
+	if (const bool bHasAnyPendingTasks = HasAnyPendingTasks(bIncludePSOPrecaching))
 	{
 		checkf(!bHasAnyPendingTasks, TEXT("UFastGeoContainer::StartRegisterTransition() was called while the container still has pending tasks. This will lead to invalid behavior."));
 		return;
@@ -837,15 +838,15 @@ void UFastGeoContainer::Tick(bool bWaitForCompletion)
 	} while (bWaitForCompletion && HasAnyPendingTasks());
 }
 
-bool UFastGeoContainer::HasAnyPendingTasks() const
+bool UFastGeoContainer::HasAnyPendingTasks(bool bIncludePSOPrecaching) const
 {
-	return HasAnyPendingCreateTasks() || HasAnyPendingDestroyTasks();
+	return HasAnyPendingCreateTasks(bIncludePSOPrecaching) || HasAnyPendingDestroyTasks();
 }
 
-bool UFastGeoContainer::HasAnyPendingCreateTasks() const
+bool UFastGeoContainer::HasAnyPendingCreateTasks(bool bIncludePSOPrecaching) const
 {
 	// Physics state can have no PendingState but OnAsyncCreatePhysicsStateEnd_GameThread has not been called yet
-	return (PhysicsState == EPhysicsState::Creating) || PendingCreate.HasAnyPendingState();
+	return (PhysicsState == EPhysicsState::Creating) || PendingCreate.HasAnyPendingState(bIncludePSOPrecaching);
 }
 
 bool UFastGeoContainer::HasAnyPendingDestroyTasks() const
@@ -1490,9 +1491,15 @@ void UFastGeoContainer::OnPrecachePSOsBegin_GameThread()
 	TRACE_CPUPROFILER_EVENT_SCOPE(UFastGeoContainer::OnPrecachePSOsBegin_GameThread);
 	check(FFastGeoStreamingModule::IsAsyncRenderWorkAllowed());
 	check(IsComponentPSOPrecachingEnabled());
-	
+
 	// PSO precaching can run before StartRegisterTransition, so initialize all prerequisites now.
 	InitializePrerequisites();
+
+	// Track async PSO precaching as pending work so that HasAnyPendingTasks() returns true.
+	// This ensures that WaitForAllPendingWorkCompletion properly gates on async PSO work completion.
+	// Set TotalNumProcessed to -1 so IsCompleted() (TotalNumProcessed >= ComponentsToProcess.Num())
+	// returns false while the async job is in flight.
+	PendingCreate.PrecachePSOState.TotalNumProcessed = -1;
 #endif
 }
 
@@ -1513,11 +1520,28 @@ void UFastGeoContainer::OnPrecachePSOs_Concurrent()
 #endif
 }
 
+void UFastGeoContainer::OnPrecachePSOsEnd_GameThread()
+{
+#if !WITH_EDITOR && UE_WITH_PSO_PRECACHING
+	// Clear the pending PSO precache state set by OnPrecachePSOsBegin_GameThread
+	// so that HasAnyPendingTasks() no longer reports PSO work as in-flight.
+	PendingCreate.PrecachePSOState.Reset();
+#endif
+}
+
 void UFastGeoContainer::PrecachePSOs()
 {
+#if !WITH_EDITOR && UE_WITH_PSO_PRECACHING
 	// This function is currently only starting precaching when running in async mode
 	TRACE_CPUPROFILER_EVENT_SCOPE(UFastGeoContainer::PrecachePSOs_GameThread);
 	StartPrecachePSOWork();
+	if (HasAnyPendingTasks())
+	{
+		// Only TickRenderWork is needed here to Launch the PSO precache job via ProcessAsyncRenderStateJobs.
+		constexpr bool bWaitForCompletion = false;
+		TickRenderWork(bWaitForCompletion);
+	}
+#endif
 }
 
 void UFastGeoContainer::ForEachMaterialAsset(TFunctionRef<void(UMaterialInterface*)> Func)
