[FastGeo] Fix surrogate re-registration during FGlobalComponentReregisterContext
- FGlobalComponentReregisterContext operates on individual UActorComponent, not actors. PostRegisterAllComponents is not called on the surrogate actor, leaving bHasRegisteredWithFastGeo out of sync.
- Add bIsReregistering flag on subsystem, set during OnGlobalComponentReregisterContext Created/Destroyed. Relax surrogate registration check in StartRegisterTransition when reregistering.
- Add ensure in OnPreRecreateScene validating that FGlobalComponentReregisterContext has already cleared all AsyncContext caches before scene recreation.
#test Open FastGeoSimple -> PIE -> Perf&Scalability -> Preview Platform -> <platform>

#Synced-CL 51624647

diff --git a/Engine/Plugins/Experimental/FastGeoStreaming/Source/FastGeoStreaming/Internal/FastGeoWorldSubsystem.h b/Engine/Plugins/Experimental/FastGeoStreaming/Source/FastGeoStreaming/Internal/FastGeoWorldSubsystem.h
--- a/Engine/Plugins/Experimental/FastGeoStreaming/Source/FastGeoStreaming/Internal/FastGeoWorldSubsystem.h
+++ b/Engine/Plugins/Experimental/FastGeoStreaming/Source/FastGeoStreaming/Internal/FastGeoWorldSubsystem.h
@@ -85,6 +85,7 @@ public:
 #endif
 
 	bool IsWaitingForCompletion() const;
+	bool IsReregistering() const { return bIsReregistering; }
 
 #if WITH_EDITOR
 public:
@@ -138,6 +139,7 @@ private:
 	float UsedAsyncRenderStateTasksTimeBudgetMS = 0;
 	int32 UsedNumComponentsToProcessBudget = 0;
 	bool bWaitingForCompletion = false;
+	bool bIsReregistering = false;
 
 	TUniquePtr<FFastGeoAsyncRenderStateJobQueue> AsyncRenderStateJobQueue;
 	TArray<TWeakObjectPtr<UFastGeoContainer>> ReregisteringFastGeoContainers;
diff --git a/Engine/Plugins/Experimental/FastGeoStreaming/Source/FastGeoStreaming/Private/FastGeoContainer.cpp b/Engine/Plugins/Experimental/FastGeoStreaming/Source/FastGeoStreaming/Private/FastGeoContainer.cpp
--- a/Engine/Plugins/Experimental/FastGeoStreaming/Source/FastGeoStreaming/Private/FastGeoContainer.cpp
+++ b/Engine/Plugins/Experimental/FastGeoStreaming/Source/FastGeoStreaming/Private/FastGeoContainer.cpp
@@ -254,7 +254,8 @@ void UFastGeoContainer::StartRegisterTransition()
 	check(!bIsBeingAsyncDestroyed);
 	check(RegistrationTargetState == ERegistrationTargetState::Registered);
 	check(RegistrationState == ERegistrationState::Unregistered);
-	check(!IsUsingSurrogateComponents() || !SurrogateActor->HasRegisteredWithFastGeo());
+	// Surrogate flag may be stale during FGlobalComponentReregisterContext (no actor-level callbacks).
+	check(!IsUsingSurrogateComponents() || !SurrogateActor->HasRegisteredWithFastGeo() || GetWorldSubsystem()->IsReregistering());
 
 	constexpr bool bIncludePSOPrecaching = false; // PSO precaching is allowed to run before registration
 	if (const bool bHasAnyPendingTasks = HasAnyPendingTasks(bIncludePSOPrecaching))
diff --git a/Engine/Plugins/Experimental/FastGeoStreaming/Source/FastGeoStreaming/Private/FastGeoWorldSubsystem.cpp b/Engine/Plugins/Experimental/FastGeoStreaming/Source/FastGeoStreaming/Private/FastGeoWorldSubsystem.cpp
--- a/Engine/Plugins/Experimental/FastGeoStreaming/Source/FastGeoStreaming/Private/FastGeoWorldSubsystem.cpp
+++ b/Engine/Plugins/Experimental/FastGeoStreaming/Source/FastGeoStreaming/Private/FastGeoWorldSubsystem.cpp
@@ -332,8 +332,9 @@ void UFastGeoWorldSubsystem::OnGlobalComponentReregisterContextCreated()
 	check(ReregisteringFastGeoContainers.IsEmpty());
 
 	WaitForAllPendingWorkCompletion();
-	
-	TGuardValue<bool> GuardValue(bWaitingForCompletion, true);
+
+	TGuardValue<bool> GuardWaitingForCompletion(bWaitingForCompletion, true);
+	TGuardValue<bool> GuardIsReregistering(bIsReregistering, true);
 	UWorld* World = GetWorld();
 	for (TObjectIterator<UFastGeoContainer> It(RF_ClassDefaultObject, true, EInternalObjectFlags::Garbage); It; ++It)
 	{
@@ -352,7 +353,8 @@ void UFastGeoWorldSubsystem::OnGlobalComponentReregisterContextDestroyed()
 {
 	WaitForAllPendingWorkCompletion();
 
-	TGuardValue<bool> GuardValue(bWaitingForCompletion, true);
+	TGuardValue<bool> GuardWaitingForCompletion(bWaitingForCompletion, true);
+	TGuardValue<bool> GuardIsReregistering(bIsReregistering, true);
 	for (const TWeakObjectPtr<UFastGeoContainer>& WeakFastGeo : ReregisteringFastGeoContainers)
 	{
 		UFastGeoContainer* FastGeo = WeakFastGeo.Get();
@@ -378,6 +380,15 @@ void UFastGeoWorldSubsystem::OnPreRecreateScene(UWorld* World)
 		// This guarantees no FastGeo render/physics jobs still reference the old scene.
 		WaitForAllPendingWorkCompletion();
 		check(!AsyncRenderStateJobQueue.IsValid() || AsyncRenderStateJobQueue->IsCompleted());
+
+		// All callers of RecreateScene() are preceded by FGlobalComponentReregisterContext,
+		// which unregisters all containers (resetting their AsyncContext). Validate this assumption.
+#if DO_CHECK
+		for (TObjectIterator<UFastGeoContainer> It(RF_ClassDefaultObject, true, EInternalObjectFlags::Garbage); It; ++It)
+		{
+			ensureMsgf((*It)->GetWorld() != World || !(*It)->AsyncContext.IsInitialized(), TEXT("Container %s still has initialized AsyncContext in OnPreRecreateScene"), *GetNameSafe(*It));
+		}
+#endif
 	}
 }
 #endif // WITH_EDITOR
