[FastGeo] Precompute spatial hashes at cook time for ISM instance hierarchy optimization
- At cook time, sort instances by spatial locality, physically reorder all instance data (transforms, custom data, random IDs), and serialize compressed spatial hashes.
- At runtime, attach precomputed spatial hashes via SetImmutable for batch insertion into the instance hierarchy. Skip if a world transform was applied at stream-in.
- Add FInstanceDataHelpers (Engine Internal) with shared BuildSpatialHashData, HashPrimitiveLocalToWorld, and GenerateInstanceRandomIDs.
- Refactor FInstanceDataManager, FPrimitiveInstanceDataManager, UInstancedStaticMeshComponent, and MassISMRenderStateHelper to use FInstanceDataHelpers.
- Move FCompressedSpatialHashItem and TLocation serialization operators to public headers.

#Synced-CL 51624647

diff --git a/Engine/Plugins/Experimental/FastGeoStreaming/Source/FastGeoStreaming/FastGeoStreaming.Build.cs b/Engine/Plugins/Experimental/FastGeoStreaming/Source/FastGeoStreaming/FastGeoStreaming.Build.cs
--- a/Engine/Plugins/Experimental/FastGeoStreaming/Source/FastGeoStreaming/FastGeoStreaming.Build.cs
+++ b/Engine/Plugins/Experimental/FastGeoStreaming/Source/FastGeoStreaming/FastGeoStreaming.Build.cs
@@ -6,6 +6,8 @@ namespace UnrealBuildTool.Rules
 	{
 		public FastGeoStreaming(ReadOnlyTargetRules Target) : base(Target)
 		{
+			PrivateIncludePaths.Add(System.IO.Path.Combine(EngineDirectory, "Source/Runtime/Engine/Internal"));
+
 			PrivateDependencyModuleNames.AddRange(
 				new string[]
 				{
diff --git a/Engine/Plugins/Experimental/FastGeoStreaming/Source/FastGeoStreaming/Internal/FastGeoInstancedStaticMeshComponent.h b/Engine/Plugins/Experimental/FastGeoStreaming/Source/FastGeoStreaming/Internal/FastGeoInstancedStaticMeshComponent.h
--- a/Engine/Plugins/Experimental/FastGeoStreaming/Source/FastGeoStreaming/Internal/FastGeoInstancedStaticMeshComponent.h
+++ b/Engine/Plugins/Experimental/FastGeoStreaming/Source/FastGeoStreaming/Internal/FastGeoInstancedStaticMeshComponent.h
@@ -66,8 +66,6 @@ protected:
 #endif
 	//~ End FFastGeoStaticMeshComponentBase interface
 
-	void BuildInstanceData();
-
 private:
 	enum class EBoundsType
 	{
@@ -76,6 +74,7 @@ private:
 		NavigationBounds
 	};
 	FBoxSphereBounds CalculateBounds(EBoundsType BoundsType);
+	void BuildInstanceData();
 	void CreateAllInstanceBodies();
 	static TArray<Chaos::FPhysicsObject*> GetAllPhysicsObjects(const TArray<FBodyInstance*>& InInstanceBodies);
 #if WITH_EDITOR
@@ -84,22 +83,24 @@ private:
 
 private:
 	// Persistent Data
+	// ------------------------------
 	TArray<FInstancedStaticMeshInstanceData> HighPrecisionPerInstanceSMData;
 	TArray<FTransform3f> LowPrecisionPerInstanceSMData;
-	bool bUseHighPrecisionPerInstanceSMData = false;
 	int32 LastInstanceBodyIndex = INDEX_NONE;
 	int32 InstancingRandomSeed = 0;
 	TArray<float> PerInstanceSMCustomData;
 	TArray<FInstancedStaticMeshRandomSeed> AdditionalRandomSeeds;
 	FBox NavigationBounds;
+	TArray<FInstanceSceneDataBuffers::FCompressedSpatialHashItem> SpatialHashes;
+	TArray<float> PerInstanceRandomIDs;
 	FInstancedStaticMeshSceneProxyDesc SceneProxyDesc{};
+	bool bUseHighPrecisionPerInstanceSMData : 1 = false;
 
 	// Transient data
-	TArray<float> InstanceRandomIDs;
-	
+	// ------------------------------
+	bool bAppliedWorldTransform : 1 = false;
 	/** Physics representation of the instance bodies. */
 	TArray<FBodyInstance*> InstanceBodies;
-
 	/** Payload used by asynchronous destruction of physics state (see OnAsyncDestroyPhysicsState). */
 	TArray<FBodyInstance*> AsyncDestroyPhysicsStatePayload;
 
diff --git a/Engine/Plugins/Experimental/FastGeoStreaming/Source/FastGeoStreaming/Private/FastGeoInstancedStaticMeshComponent.cpp b/Engine/Plugins/Experimental/FastGeoStreaming/Source/FastGeoStreaming/Private/FastGeoInstancedStaticMeshComponent.cpp
--- a/Engine/Plugins/Experimental/FastGeoStreaming/Source/FastGeoStreaming/Private/FastGeoInstancedStaticMeshComponent.cpp
+++ b/Engine/Plugins/Experimental/FastGeoStreaming/Source/FastGeoStreaming/Private/FastGeoInstancedStaticMeshComponent.cpp
@@ -12,10 +12,12 @@
 #include "AI/Navigation/NavigationRelevantData.h"
 #include "Engine/InstancedStaticMesh.h"
 #include "InstancedStaticMeshComponentHelper.h"
+#include "InstanceData/InstanceDataHelpers.h"
 #include "InstancedStaticMesh/ISMInstanceDataSceneProxy.h"
 #include "NaniteVertexFactory.h"
 #include "NaniteSceneProxy.h"
 #include "PhysicsEngine/BodySetup.h"
+#include "Math/DoubleFloat.h"
 #include "SceneInterface.h"
 #include "Physics/Experimental/PhysScene_Chaos.h"
 #include "PhysicsEngine/PhysicsObjectExternalInterface.h"
@@ -33,7 +35,7 @@ void FFastGeoInstancedStaticMeshComponent::Serialize(FArchive& Ar)
 	Super::Serialize(Ar);
 
 	// Serialize persistent data from FFastGeoInstancedStaticMeshComponent
-	Ar << bUseHighPrecisionPerInstanceSMData;
+	FArchive_Serialize_BitfieldBool(Ar, bUseHighPrecisionPerInstanceSMData);
 	if (bUseHighPrecisionPerInstanceSMData)
 	{
 		HighPrecisionPerInstanceSMData.BulkSerialize(Ar);
@@ -54,6 +56,10 @@ void FFastGeoInstancedStaticMeshComponent::Serialize(FArchive& Ar)
 	Ar << SceneProxyDesc.InstanceStartCullDistance;
 	Ar << SceneProxyDesc.InstanceEndCullDistance;
 	FArchive_Serialize_BitfieldBool(Ar, SceneProxyDesc.bUseGpuLodSelection);
+
+	// Precomputed spatial hashes and pre-generated random IDs for instance hierarchy optimization
+	SpatialHashes.BulkSerialize(Ar);
+	PerInstanceRandomIDs.BulkSerialize(Ar);
 }
 
 #if WITH_EDITOR
@@ -80,6 +86,100 @@ void FFastGeoInstancedStaticMeshComponent::InitializeFromComponent(UActorCompone
 	LocalBounds = CalculateBounds(EBoundsType::LocalBounds);
 	WorldBounds = CalculateBounds(EBoundsType::WorldBounds);
 	NavigationBounds = CalculateBounds(EBoundsType::NavigationBounds).GetBox();
+
+	// Compute spatial hash order and physically reorder instance data for optimal RT performance.
+	// This sorts instances by spatial locality so the render thread can use precomputed spatial
+	// hashes for batch insertion into the instance hierarchy instead of per-instance processing.
+	{
+		const int32 InstanceCount = GetInstanceCount();
+		const FRenderBounds MeshBounds = GetStaticMesh()->GetBounds();
+		const FMatrix ComponentTransformMatrix = WorldTransform.ToMatrixWithScale();
+
+		// Compute PrimitiveToRelativeWorld and world space offset (same as FInstanceDataManager does)
+		const FVector3f PrimitiveWorldSpacePositionHigh = FDFVector3{ ComponentTransformMatrix.GetOrigin() }.High;
+		const FRenderTransform PrimitiveToRelativeWorld = FDFMatrix::MakeToRelativeWorldMatrix(PrimitiveWorldSpacePositionHigh, ComponentTransformMatrix).M;
+		const FVector PrimitiveWorldSpaceOffset(PrimitiveWorldSpacePositionHigh);
+
+		// Build spatial hash data and get reorder table (empty if already in optimal order)
+		FInstanceDataHelpers::FSpatialHashResult SpatialHashResult = FInstanceDataHelpers::BuildSpatialHashData(InstanceCount, [&](int32 InstanceIndex) -> FSphere
+		{
+			FMatrix InstanceTransform;
+			if (bUseHighPrecisionPerInstanceSMData)
+			{
+				InstanceTransform = HighPrecisionPerInstanceSMData[InstanceIndex].Transform;
+			}
+			else
+			{
+				InstanceTransform = FMatrix(LowPrecisionPerInstanceSMData[InstanceIndex].ToMatrixWithScale());
+			}
+			const FSphere3f LocalSphere(MeshBounds.GetCenter(), MeshBounds.GetExtent().Size());
+			const FRenderTransform InstanceLocalToWorld = FRenderTransform(InstanceTransform) * PrimitiveToRelativeWorld;
+			FSphere Result(LocalSphere.TransformBy(InstanceLocalToWorld.ToMatrix44f()));
+			Result.Center += PrimitiveWorldSpaceOffset;
+			return Result;
+		});
+
+		SpatialHashes = MoveTemp(SpatialHashResult.SpatialHashes);
+		const TArray<int32>& ReorderTable = SpatialHashResult.ReorderTable;
+		const bool bNeedsReorder = ReorderTable.Num() > 0;
+
+		// Physically reorder instance data arrays
+		if (bNeedsReorder)
+		{
+			if (bUseHighPrecisionPerInstanceSMData)
+			{
+				TArray<FInstancedStaticMeshInstanceData> Reordered;
+				Reordered.SetNumUninitialized(InstanceCount);
+				for (int32 Index = 0; Index < InstanceCount; ++Index)
+				{
+					Reordered[Index] = HighPrecisionPerInstanceSMData[ReorderTable[Index]];
+				}
+				HighPrecisionPerInstanceSMData = MoveTemp(Reordered);
+			}
+			else
+			{
+				TArray<FTransform3f> Reordered;
+				Reordered.SetNumUninitialized(InstanceCount);
+				for (int32 Index = 0; Index < InstanceCount; ++Index)
+				{
+					Reordered[Index] = LowPrecisionPerInstanceSMData[ReorderTable[Index]];
+				}
+				LowPrecisionPerInstanceSMData = MoveTemp(Reordered);
+			}
+
+			// Reorder custom data (strided by NumCustomDataFloats)
+			const int32 NumCustomDataFloats = PerInstanceSMCustomData.Num() / FMath::Max(InstanceCount, 1);
+			if (NumCustomDataFloats > 0)
+			{
+				const SIZE_T CustomDataStride = NumCustomDataFloats * sizeof(float);
+				TArray<float> ReorderedCustomData;
+				ReorderedCustomData.SetNumUninitialized(PerInstanceSMCustomData.Num());
+				for (int32 Index = 0; Index < InstanceCount; ++Index)
+				{
+					const int32 OldIndex = ReorderTable[Index];
+					FMemory::Memcpy(&ReorderedCustomData[Index * NumCustomDataFloats], &PerInstanceSMCustomData[OldIndex * NumCustomDataFloats], CustomDataStride);
+				}
+				PerInstanceSMCustomData = MoveTemp(ReorderedCustomData);
+			}
+		}
+
+		// Pre-generate random IDs in the reordered order
+		{
+			TArray<float> OriginalRandomIDs = FInstanceDataHelpers::GenerateInstanceRandomIDs(InstanceCount, InstancingRandomSeed, AdditionalRandomSeeds);
+			if (bNeedsReorder)
+			{
+				PerInstanceRandomIDs.SetNumUninitialized(InstanceCount);
+				for (int32 Index = 0; Index < InstanceCount; ++Index)
+				{
+					PerInstanceRandomIDs[Index] = OriginalRandomIDs[ReorderTable[Index]];
+				}
+			}
+			else
+			{
+				PerInstanceRandomIDs = MoveTemp(OriginalRandomIDs);
+			}
+		}
+	}
 }
 
 void FFastGeoInstancedStaticMeshComponent::InitializePerInstanceSMData(const TArray<FInstancedStaticMeshInstanceData>& InPerInstanceSMData)
@@ -147,6 +247,7 @@ void FFastGeoInstancedStaticMeshComponent::ApplyWorldTransform(const FTransform&
 {
 	Super::ApplyWorldTransform(InTransform);
 
+	bAppliedWorldTransform = true;
 	WorldBounds = CalculateBounds(EBoundsType::WorldBounds);
 	NavigationBounds = CalculateBounds(EBoundsType::NavigationBounds).GetBox();
 }
@@ -178,11 +279,11 @@ void FFastGeoInstancedStaticMeshComponent::BuildInstanceData()
 {
 	FInstanceSceneDataBuffers InstanceSceneDataBuffers{};
 	FInstanceSceneDataBuffers::FAccessTag AccessTag(PointerHash(this));
-	auto View = InstanceSceneDataBuffers.BeginWriteAccess(AccessTag);
-	
+	FInstanceSceneDataBuffers::FWriteView View = InstanceSceneDataBuffers.BeginWriteAccess(AccessTag);
+
 	// PrimitiveLocalToWorld
 	InstanceSceneDataBuffers.SetPrimitiveLocalToWorld(GetRenderMatrix(), AccessTag);
-	
+
 	// InstanceLocalBounds
 	const FPrimitiveMaterialPropertyDescriptor PrimitiveMaterialDesc = GetUsedMaterialPropertyDesc(GetScene()->GetShaderPlatform());
 	const float LocalAbsMaxDisplacement = FMath::Max(-PrimitiveMaterialDesc.MinMaxMaterialDisplacement.X, PrimitiveMaterialDesc.MinMaxMaterialDisplacement.Y) + PrimitiveMaterialDesc.MaxWorldPositionOffsetDisplacement;
@@ -196,45 +297,50 @@ void FFastGeoInstancedStaticMeshComponent::BuildInstanceData()
 	// LocalToPrimitiveRelativeWorld
 	const int32 InstanceCount = GetInstanceCount();
 	View.InstanceToPrimitiveRelative.Reserve(InstanceCount);
-	ForEachInstanceMatrix([&View](const FMatrix& InstanceTransform)
+	if (bUseHighPrecisionPerInstanceSMData)
 	{
-		FRenderTransform InstanceToPrimitive = InstanceTransform;
-		FRenderTransform LocalToPrimitiveRelativeWorld = InstanceToPrimitive * View.PrimitiveToRelativeWorld;
-		LocalToPrimitiveRelativeWorld.Orthogonalize();
-		View.InstanceToPrimitiveRelative.Add(LocalToPrimitiveRelativeWorld);
-	});
+		for (const FInstancedStaticMeshInstanceData& InstanceData : HighPrecisionPerInstanceSMData)
+		{
+			FRenderTransform LocalToPrimitiveRelativeWorld = FRenderTransform(InstanceData.Transform) * View.PrimitiveToRelativeWorld;
+			LocalToPrimitiveRelativeWorld.Orthogonalize();
+			View.InstanceToPrimitiveRelative.Add(LocalToPrimitiveRelativeWorld);
+		}
+	}
+	else
+	{
+		for (const FTransform3f& Transform : LowPrecisionPerInstanceSMData)
+		{
+			FRenderTransform LocalToPrimitiveRelativeWorld = FRenderTransform(Transform.ToMatrixWithScale()) * View.PrimitiveToRelativeWorld;
+			LocalToPrimitiveRelativeWorld.Orthogonalize();
+			View.InstanceToPrimitiveRelative.Add(LocalToPrimitiveRelativeWorld);
+		}
+	}
 
 	// InstanceCustomData
+	const int32 NumCustomDataFloats = (PerInstanceSMCustomData.Num() && InstanceCount) ? PerInstanceSMCustomData.Num() / InstanceCount : 0;
 	View.InstanceCustomData = PerInstanceSMCustomData;
-	View.NumCustomDataFloats = PerInstanceSMCustomData.Num() && InstanceCount ? PerInstanceSMCustomData.Num() / InstanceCount : 0;
-	View.Flags.bHasPerInstanceCustomData = PrimitiveMaterialDesc.bAnyMaterialHasPerInstanceCustomData && View.NumCustomDataFloats != 0;
+	View.NumCustomDataFloats = NumCustomDataFloats;
+	View.Flags.bHasPerInstanceCustomData = PrimitiveMaterialDesc.bAnyMaterialHasPerInstanceCustomData && NumCustomDataFloats != 0;
 	if (!View.Flags.bHasPerInstanceCustomData)
 	{
 		View.NumCustomDataFloats = 0;
 		View.InstanceCustomData.Reset();
 	}
 
-	// InstanceRandomIDs
+	// InstanceRandomIDs (pre-generated in spatial hash order from cook)
 	View.Flags.bHasPerInstanceRandom = PrimitiveMaterialDesc.bAnyMaterialHasPerInstanceRandom && (InstanceCount > 0);
-	if (View.Flags.bHasPerInstanceRandom && InstanceRandomIDs.IsEmpty())
+	if (View.Flags.bHasPerInstanceRandom)
 	{
-		InstanceRandomIDs.SetNumZeroed(InstanceCount);
-		FRandomStream RandomStream = FRandomStream(InstancingRandomSeed);
-		auto AdditionalRandomSeedsIt = AdditionalRandomSeeds.CreateIterator();
-		int32 SeedResetIndex = AdditionalRandomSeedsIt ? AdditionalRandomSeedsIt->StartInstanceIndex : INDEX_NONE;
-		for (int32 Index = 0; Index < InstanceRandomIDs.Num(); ++Index)
-		{
-			// Reset the random stream if necessary
-			if (Index == SeedResetIndex)
-			{
-				RandomStream = FRandomStream(AdditionalRandomSeedsIt->RandomSeed);
-				AdditionalRandomSeedsIt++;
-				SeedResetIndex = AdditionalRandomSeedsIt ? AdditionalRandomSeedsIt->StartInstanceIndex : INDEX_NONE;
-			}
-			InstanceRandomIDs[Index] = RandomStream.GetFraction();
-		}
+		View.InstanceRandomIDs = PerInstanceRandomIDs;
+	}
+
+	// Precomputed spatial hashes (no reorder table - data already reordered at cook time).
+	// If a world transform was applied at stream-in, spatial hash cell locations are invalid
+	// since they encode cook-time world positions. Reordered data still provides cache locality.
+	if (!SpatialHashes.IsEmpty() && !bAppliedWorldTransform)
+	{
+		InstanceSceneDataBuffers.SetImmutable(FInstanceSceneDataImmutable(SpatialHashes), AccessTag);
 	}
-	View.InstanceRandomIDs = InstanceRandomIDs;
 
 	InstanceSceneDataBuffers.EndWriteAccess(AccessTag);
 	InstanceSceneDataBuffers.ValidateData();
diff --git a/Engine/Source/Runtime/Engine/Internal/InstanceData/InstanceDataHelpers.h b/Engine/Source/Runtime/Engine/Internal/InstanceData/InstanceDataHelpers.h
new file mode 100644
--- /dev/null
+++ b/Engine/Source/Runtime/Engine/Internal/InstanceData/InstanceDataHelpers.h
@@ -0,0 +1,63 @@
+// Copyright Epic Games, Inc. All Rights Reserved.
+
+#pragma once
+
+#include "CoreTypes.h"
+#include "RenderTransform.h"
+
+#if WITH_EDITOR
+#include "InstanceDataSceneProxy.h"
+#endif
+
+struct FInstancedStaticMeshRandomSeed;
+
+/**
+ * Shared utility functions for instance data processing.
+ * Used by FInstanceDataManager, FPrimitiveInstanceDataManager, and FastGeo to avoid code duplication.
+ */
+class FInstanceDataHelpers
+{
+public:
+	/**
+	 * Hash a primitive transform for spatial hash validation.
+	 * Rounds components to int to reduce floating point precision sensitivity.
+	 * Used to detect whether a primitive's transform has changed since spatial hashes were computed.
+	 */
+	ENGINE_API static uint32 HashPrimitiveLocalToWorld(const FRenderTransform& PrimitiveToRelativeWorld, const FVector& PrimitiveWorldSpaceOffset);
+
+	/**
+	 * Generate per-instance random IDs using the same deterministic algorithm as UInstancedStaticMeshComponent.
+	 * Iterates through instances, advancing a random stream seeded with InstancingRandomSeed,
+	 * and resetting the stream at indices specified by AdditionalRandomSeeds.
+	 * Fills the provided array in-place. The array must be pre-sized to the desired instance count.
+	 */
+	ENGINE_API static void GenerateInstanceRandomIDs(int32 InstancingRandomSeed, const TArray<FInstancedStaticMeshRandomSeed>& AdditionalRandomSeeds, TArrayView<float> OutRandomIDs);
+
+	/**
+	 * Generate per-instance random IDs and return them as a new array.
+	 * @see GenerateInstanceRandomIDs(int32, const TArray<FInstancedStaticMeshRandomSeed>&, TArrayView<float>)
+	 */
+	ENGINE_API static TArray<float> GenerateInstanceRandomIDs(int32 InstanceCount, int32 InstancingRandomSeed, const TArray<FInstancedStaticMeshRandomSeed>& AdditionalRandomSeeds);
+
+#if WITH_EDITOR
+	/** Result of BuildSpatialHashData, containing the compressed spatial hashes and the instance reorder table. */
+	struct FSpatialHashResult
+	{
+		/** Compressed spatial hash items. Each item represents a contiguous range of instances sharing the same spatial hash cell. */
+		TArray<FInstanceSceneDataBuffers::FCompressedSpatialHashItem> SpatialHashes;
+		/** Reorder table mapping new index -> original index. Empty if instances were already in optimal order. */
+		TArray<int32> ReorderTable;
+	};
+
+	/**
+	 * Build compressed spatial hash data from instance world-space spheres.
+	 * Sorts instances by spatial locality using FSpatialHashSortBuilder, then compresses
+	 * consecutive same-location instances into FCompressedSpatialHashItem ranges.
+	 *
+	 * @param NumInstances                  Number of instances to process.
+	 * @param GetWorldSpaceInstanceSphere   Callback returning the world-space bounding sphere for each instance index.
+	 * @return Compressed spatial hash items and reorder table.
+	 */
+	ENGINE_API static FSpatialHashResult BuildSpatialHashData(int32 NumInstances, TFunctionRef<FSphere(int32)> GetWorldSpaceInstanceSphere);
+#endif
+};
diff --git a/Engine/Source/Runtime/Engine/Private/InstanceData/InstanceDataHelpers.cpp b/Engine/Source/Runtime/Engine/Private/InstanceData/InstanceDataHelpers.cpp
new file mode 100644
--- /dev/null
+++ b/Engine/Source/Runtime/Engine/Private/InstanceData/InstanceDataHelpers.cpp
@@ -0,0 +1,111 @@
+// Copyright Epic Games, Inc. All Rights Reserved.
+
+#include "InstanceData/InstanceDataHelpers.h"
+#include "Engine/InstancedStaticMesh.h"
+#include "Math/UnrealMathUtility.h"
+#include "Templates/TypeHash.h"
+
+#if WITH_EDITOR
+#include "InstanceData/InstanceDataUpdateUtils.h"
+#endif
+
+uint32 FInstanceDataHelpers::HashPrimitiveLocalToWorld(const FRenderTransform& M, const FVector& PrimitiveWorldSpaceOffset)
+{
+	// Round to reduce risk of floating point precision causing otherwise same transforms from being considered different
+	return HashCombineFast(
+		GetTypeHash(FMath::RoundToInt(M.TransformRows[0][0])), GetTypeHash(FMath::RoundToInt(M.TransformRows[0][1])), GetTypeHash(FMath::RoundToInt(M.TransformRows[0][2])),
+		GetTypeHash(FMath::RoundToInt(M.TransformRows[1][0])), GetTypeHash(FMath::RoundToInt(M.TransformRows[1][1])), GetTypeHash(FMath::RoundToInt(M.TransformRows[1][2])),
+		GetTypeHash(FMath::RoundToInt(M.TransformRows[2][0])), GetTypeHash(FMath::RoundToInt(M.TransformRows[2][1])), GetTypeHash(FMath::RoundToInt(M.TransformRows[2][2])),
+		GetTypeHash(FMath::RoundToInt(M.Origin[0])), GetTypeHash(FMath::RoundToInt(M.Origin[1])), GetTypeHash(FMath::RoundToInt(M.Origin[2])),
+		GetTypeHash(FMath::RoundToInt(PrimitiveWorldSpaceOffset[0])), GetTypeHash(FMath::RoundToInt(PrimitiveWorldSpaceOffset[1])), GetTypeHash(FMath::RoundToInt(PrimitiveWorldSpaceOffset[2]))
+	);
+}
+
+void FInstanceDataHelpers::GenerateInstanceRandomIDs(int32 InstancingRandomSeed, const TArray<FInstancedStaticMeshRandomSeed>& AdditionalRandomSeeds, TArrayView<float> OutRandomIDs)
+{
+	FRandomStream RandomStream(InstancingRandomSeed);
+	TArray<FInstancedStaticMeshRandomSeed>::TConstIterator SeedsIt = AdditionalRandomSeeds.CreateConstIterator();
+	int32 SeedResetIndex = SeedsIt ? SeedsIt->StartInstanceIndex : INDEX_NONE;
+	for (int32 Index = 0; Index < OutRandomIDs.Num(); ++Index)
+	{
+		if (Index == SeedResetIndex)
+		{
+			RandomStream = FRandomStream(SeedsIt->RandomSeed);
+			++SeedsIt;
+			SeedResetIndex = SeedsIt ? SeedsIt->StartInstanceIndex : INDEX_NONE;
+		}
+		OutRandomIDs[Index] = RandomStream.GetFraction();
+	}
+}
+
+TArray<float> FInstanceDataHelpers::GenerateInstanceRandomIDs(int32 InstanceCount, int32 InstancingRandomSeed, const TArray<FInstancedStaticMeshRandomSeed>& AdditionalRandomSeeds)
+{
+	TArray<float> RandomIDs;
+	RandomIDs.SetNumUninitialized(InstanceCount);
+	GenerateInstanceRandomIDs(InstancingRandomSeed, AdditionalRandomSeeds, RandomIDs);
+	return RandomIDs;
+}
+
+#if WITH_EDITOR
+
+FInstanceDataHelpers::FSpatialHashResult FInstanceDataHelpers::BuildSpatialHashData(int32 NumInstances, TFunctionRef<FSphere(int32)> GetWorldSpaceInstanceSphere)
+{
+	int32 MinLevel = 0;
+	static const auto CVarInstanceHierarchyMinCellSize = IConsoleManager::Get().FindTConsoleVariableDataFloat(TEXT("r.SceneCulling.MinCellSize"));
+	if (CVarInstanceHierarchyMinCellSize)
+	{
+		MinLevel = RenderingSpatialHash::CalcLevel(CVarInstanceHierarchyMinCellSize->GetValueOnAnyThread() - 1.0);
+	}
+
+	FSpatialHashSortBuilder SortBuilder;
+	SortBuilder.BuildOptimizedSpatialHashOrder(NumInstances, MinLevel, GetWorldSpaceInstanceSphere);
+
+	// Build reorder table and compress spatial hashes in a single pass over the sorted instances.
+	// This logic is extracted from FInstanceDataManager::PrecomputeOptimizationData and
+	// FPrimitiveInstanceDataManager::PrecomputeOptimizationData with no behavioral changes.
+	FSpatialHashResult Result;
+	Result.ReorderTable.SetNumUninitialized(NumInstances);
+
+	FInstanceSceneDataBuffers::FCompressedSpatialHashItem CurrentItem;
+	CurrentItem.NumInstances = 0;
+
+	bool bIsIdentityIndexMap = true;
+
+	for (int32 Index = 0; Index < SortBuilder.SortedInstances.Num(); ++Index)
+	{
+		// Reorder table
+		int32 ComponentInstanceIndex = SortBuilder.SortedInstances[Index].InstanceIndex;
+		bIsIdentityIndexMap = bIsIdentityIndexMap && (Index == ComponentInstanceIndex);
+		Result.ReorderTable[Index] = ComponentInstanceIndex;
+
+		// Compress consecutive same-location instances into ranges
+		bool bSameLoc = CurrentItem.NumInstances > 0 && CurrentItem.Location == SortBuilder.SortedInstances[Index].InstanceLoc;
+		if (bSameLoc)
+		{
+			CurrentItem.NumInstances += 1;
+		}
+		else
+		{
+			if (CurrentItem.NumInstances > 0)
+			{
+				Result.SpatialHashes.Add(CurrentItem);
+			}
+			CurrentItem.Location = SortBuilder.SortedInstances[Index].InstanceLoc;
+			CurrentItem.NumInstances = 1;
+		}
+	}
+	if (CurrentItem.NumInstances > 0)
+	{
+		Result.SpatialHashes.Add(CurrentItem);
+	}
+
+	// Don't store a 1:1 mapping
+	if (bIsIdentityIndexMap)
+	{
+		Result.ReorderTable.Reset();
+	}
+
+	return Result;
+}
+
+#endif
\ No newline at end of file
diff --git a/Engine/Source/Runtime/Engine/Private/InstanceData/InstanceDataManager.cpp b/Engine/Source/Runtime/Engine/Private/InstanceData/InstanceDataManager.cpp
--- a/Engine/Source/Runtime/Engine/Private/InstanceData/InstanceDataManager.cpp
+++ b/Engine/Source/Runtime/Engine/Private/InstanceData/InstanceDataManager.cpp
@@ -1,6 +1,7 @@
 // Copyright Epic Games, Inc. All Rights Reserved.
 
 #include "InstanceData/InstanceDataManager.h"
+#include "InstanceData/InstanceDataHelpers.h"
 #include "InstanceData/InstanceUpdateChangeSet.h"
 #include "InstanceData/InstanceDataUpdateUtils.h"
 #include "Components/PrimitiveComponent.h"
@@ -29,24 +30,6 @@ static TAutoConsoleVariable<bool> CVarInstanceDataManagerBuildOptimizedForTestin
 
 
 
-/**
- * Helper function to hash the local to world that was used to build the precomputed data as we can't rely on it being unchanged, 
- * for example when sublevels with transforms are used.
- */
-inline uint32 HashPrimitiveLocalToWorld(const FRenderTransform& M, const FVector& PrimitiveWorldSpaceOffset)
-{
-	// round to reduce risk of floating point precision causing otherwise same transforms from being considered different
-	uint32 Hash = HashCombineFast(
-		GetTypeHash(FMath::RoundToInt(M.TransformRows[0][0])), GetTypeHash(FMath::RoundToInt(M.TransformRows[0][1])), GetTypeHash(FMath::RoundToInt(M.TransformRows[0][2])),
-		GetTypeHash(FMath::RoundToInt(M.TransformRows[1][0])), GetTypeHash(FMath::RoundToInt(M.TransformRows[1][1])), GetTypeHash(FMath::RoundToInt(M.TransformRows[1][2])),
-		GetTypeHash(FMath::RoundToInt(M.TransformRows[2][0])), GetTypeHash(FMath::RoundToInt(M.TransformRows[2][1])), GetTypeHash(FMath::RoundToInt(M.TransformRows[2][2])),
-		GetTypeHash(FMath::RoundToInt(M.Origin[0])), GetTypeHash(FMath::RoundToInt(M.Origin[1])), GetTypeHash(FMath::RoundToInt(M.Origin[2])),
-		GetTypeHash(FMath::RoundToInt(PrimitiveWorldSpaceOffset[0])), GetTypeHash(FMath::RoundToInt(PrimitiveWorldSpaceOffset[1])), GetTypeHash(FMath::RoundToInt(PrimitiveWorldSpaceOffset[2]))
-	);
-
-	return Hash;
-}
-
 template <typename IndexRemapType>
 void ApplyDataChanges(FInstanceUpdateChangeSet& ChangeSet, const IndexRemapType& IndexRemap, int32 PostUpdateNumInstances, FInstanceSceneDataBuffers::FWriteView& ProxyData)
 {
@@ -229,8 +212,9 @@ protected:
 #endif
 };
 
-FInstanceDataManager::FInstanceDataManager(UPrimitiveComponent* InPrimitiveComponent) 
+FInstanceDataManager::FInstanceDataManager(UPrimitiveComponent* InPrimitiveComponent)
 	: bNumCustomDataChanged(false)
+	, bTransformChangedAllInstances(false)
 	, bPrimitiveTransformChanged(false)
 	, bAnyInstanceChange(false)
 	, PrimitiveComponent(InPrimitiveComponent) 
@@ -778,7 +762,7 @@ bool FInstanceDataManager::FlushChanges(FInstanceDataManagerSourceDataDesc&& Com
 			LOG_INST_DATA(TEXT("Full Build %s"), TEXT(""));
 			if (PrecomputedOptimizationData.IsValid())
 			{
-				if (PrecomputedOptimizationData->PrimitiveLocalToWorldHash == HashPrimitiveLocalToWorld(ChangeSet.PrimitiveToRelativeWorld, ChangeSet.PrimitiveWorldSpaceOffset))
+				if (PrecomputedOptimizationData->PrimitiveLocalToWorldHash == FInstanceDataHelpers::HashPrimitiveLocalToWorld(ChangeSet.PrimitiveToRelativeWorld, ChangeSet.PrimitiveWorldSpaceOffset))
 				{
 					LOG_INST_DATA(TEXT("  Optimized Build (%s)"), PrecomputedOptimizationData.IsValid() ? TEXT("Precomputed") : TEXT(""));
 					ChangeSet.PrecomputedOptimizationData = PrecomputedOptimizationData;
@@ -1028,17 +1012,7 @@ FPrecomputedInstanceSpatialHashData FInstanceDataManager::PrecomputeOptimization
 
 FPrecomputedInstanceSpatialHashData FInstanceDataManager::PrecomputeOptimizationData(FInstanceUpdateChangeSet& ChangeSet)
 {
-	FSpatialHashSortBuilder SortBuilder;
-
-	int32 MinLevel = 0;
-	static const auto CVarInstanceHierarchyMinCellSize = IConsoleManager::Get().FindTConsoleVariableDataFloat(TEXT("r.SceneCulling.MinCellSize"));
-	if (CVarInstanceHierarchyMinCellSize)
-	{
-		// TODO: only one code path to compute this value!!!
-		MinLevel = RenderingSpatialHash::CalcLevel(CVarInstanceHierarchyMinCellSize->GetValueOnAnyThread() - 1.0);
-	}
-
-	SortBuilder.BuildOptimizedSpatialHashOrder(ChangeSet.NumSourceInstances, MinLevel,
+	FInstanceDataHelpers::FSpatialHashResult SpatialHashResult = FInstanceDataHelpers::BuildSpatialHashData(ChangeSet.NumSourceInstances,
 		[&](int32 InstanceIndex) -> FSphere
 		{
 			int32 BoundsIndex = FMath::Min(ChangeSet.InstanceLocalBounds.Num() - 1, InstanceIndex);
@@ -1047,52 +1021,14 @@ FPrecomputedInstanceSpatialHashData FInstanceDataManager::PrecomputeOptimization
 			FRenderTransform InstanceLocalToWorld = ChangeSet.Transforms[InstanceIndex] * ChangeSet.PrimitiveToRelativeWorld;
 			FSphere Result = FSphere(LocalInstanceSphere.TransformBy(InstanceLocalToWorld.ToMatrix44f()));
 			Result.Center += ChangeSet.PrimitiveWorldSpaceOffset;
-			return Result; 
-		}				
+			return Result;
+		}
 	);
 
 	FPrecomputedInstanceSpatialHashData Result;
-
-	// Pack down the spatial hashes & index remap
-	Result.ProxyIndexToComponentIndexRemap.SetNumUninitialized(ChangeSet.NumSourceInstances);
-
-	FInstanceSceneDataBuffers::FCompressedSpatialHashItem CurrentItem;
-	CurrentItem.NumInstances = 0;
-
-	bool bIsIdentityIndexMap = true;
-
-	for (int32 InstanceIndex = 0; InstanceIndex < SortBuilder.SortedInstances.Num(); ++InstanceIndex)
-	{
-		int32 ComponentInstanceIndex = SortBuilder.SortedInstances[InstanceIndex].InstanceIndex;
-		bIsIdentityIndexMap = bIsIdentityIndexMap && InstanceIndex == ComponentInstanceIndex;
-		Result.ProxyIndexToComponentIndexRemap[InstanceIndex] = ComponentInstanceIndex;
-
-		bool bSameLoc = CurrentItem.NumInstances > 0 && CurrentItem.Location == SortBuilder.SortedInstances[InstanceIndex].InstanceLoc;
-		if (bSameLoc)
-		{
-			CurrentItem.NumInstances += 1;
-		}
-		else
-		{
-			if (CurrentItem.NumInstances > 0)
-			{
-				Result.Hashes.Add(CurrentItem);
-			}
-			CurrentItem.Location = SortBuilder.SortedInstances[InstanceIndex].InstanceLoc;
-			CurrentItem.NumInstances = 1;
-		}
-	}
-	if (CurrentItem.NumInstances > 0)
-	{
-		Result.Hashes.Add(CurrentItem);
-	}
-
-	// Don't store a 1:1 mapping
-	if (bIsIdentityIndexMap)
-	{
-		Result.ProxyIndexToComponentIndexRemap.Reset();
-	}
-	Result.PrimitiveLocalToWorldHash = HashPrimitiveLocalToWorld(ChangeSet.PrimitiveToRelativeWorld, ChangeSet.PrimitiveWorldSpaceOffset);
+	Result.Hashes = MoveTemp(SpatialHashResult.SpatialHashes);
+	Result.ProxyIndexToComponentIndexRemap = MoveTemp(SpatialHashResult.ReorderTable);
+	Result.PrimitiveLocalToWorldHash = FInstanceDataHelpers::HashPrimitiveLocalToWorld(ChangeSet.PrimitiveToRelativeWorld, ChangeSet.PrimitiveWorldSpaceOffset);
 	return Result;
 }
 
diff --git a/Engine/Source/Runtime/Engine/Private/InstanceData/InstanceDataUpdateUtils.h b/Engine/Source/Runtime/Engine/Private/InstanceData/InstanceDataUpdateUtils.h
--- a/Engine/Source/Runtime/Engine/Private/InstanceData/InstanceDataUpdateUtils.h
+++ b/Engine/Source/Runtime/Engine/Private/InstanceData/InstanceDataUpdateUtils.h
@@ -435,25 +435,6 @@ void DispatchInstanceDataUpdateTask(bool bIsUnattached, const TSharedPtr<ProxyTy
 	}
 }
 
-namespace RenderingSpatialHash
-{
-	template <typename ScalarType>
-	inline FArchive& operator<<(FArchive& Ar, TLocation<ScalarType>& Item)
-	{
-		Ar << Item.Coord;
-		Ar << Item.Level;
-		return Ar;
-	}
-}
-
-inline FArchive& operator<<(FArchive& Ar, FInstanceSceneDataBuffers::FCompressedSpatialHashItem& Item)
-{
-	Ar << Item.Location;
-	Ar << Item.NumInstances;
-
-	return Ar;
-}
-
 #if WITH_EDITOR
 
 class FSpatialHashSortBuilder
diff --git a/Engine/Source/Runtime/Engine/Private/InstancedStaticMesh.cpp b/Engine/Source/Runtime/Engine/Private/InstancedStaticMesh.cpp
--- a/Engine/Source/Runtime/Engine/Private/InstancedStaticMesh.cpp
+++ b/Engine/Source/Runtime/Engine/Private/InstancedStaticMesh.cpp
@@ -6,6 +6,7 @@
 
 #include "Engine/InstancedStaticMesh.h"
 #include "Elements/SMInstance/SMInstanceElementId.h"
+#include "InstanceData/InstanceDataHelpers.h"
 #include "Engine/Level.h"
 #include "Engine/OverlapResult.h"
 #include "AI/NavigationSystemBase.h"
@@ -2742,26 +2743,12 @@ void UInstancedStaticMeshComponent::BuildInstanceDataDeltaChangeSetCommon(FISMIn
 
 		// Does this always have to crunch through all the instances from start to end each time anything changes? 
 		check(InstancingRandomSeed != 0);
-		ChangeSet.GeneratePerInstanceRandomIds = 
+		ChangeSet.GeneratePerInstanceRandomIds =
 			[InstancingRandomSeed = InstancingRandomSeed,
-			AdditionalRandomSeeds = AdditionalRandomSeeds](TArray<float> &InstanceRandomIDs) mutable
+			AdditionalRandomSeeds = AdditionalRandomSeeds](TArray<float> &InstanceRandomIDs)
 			{
-			FRandomStream RandomStream = FRandomStream(InstancingRandomSeed);
-			auto AdditionalRandomSeedsIt = AdditionalRandomSeeds.CreateIterator();
-			int32 SeedResetIndex = AdditionalRandomSeedsIt ? AdditionalRandomSeedsIt->StartInstanceIndex : INDEX_NONE;
-			for (int32 Index = 0; Index < InstanceRandomIDs.Num(); ++Index)
-			{
-				// Reset the random stream if necessary
-				if (Index == SeedResetIndex)
-				{
-					RandomStream = FRandomStream(AdditionalRandomSeedsIt->RandomSeed);
-					AdditionalRandomSeedsIt++;
-					SeedResetIndex = AdditionalRandomSeedsIt ? AdditionalRandomSeedsIt->StartInstanceIndex : INDEX_NONE;
-			}
-
-				InstanceRandomIDs[Index] = RandomStream.GetFraction();
-		}
-		};
+				FInstanceDataHelpers::GenerateInstanceRandomIDs(InstancingRandomSeed, AdditionalRandomSeeds, InstanceRandomIDs);
+			};
 	}
 
 	if (ChangeSet.Flags.bHasPerInstanceLMSMUVBias)
diff --git a/Engine/Source/Runtime/Engine/Private/InstancedStaticMesh/ISMInstanceDataManager.cpp b/Engine/Source/Runtime/Engine/Private/InstancedStaticMesh/ISMInstanceDataManager.cpp
--- a/Engine/Source/Runtime/Engine/Private/InstancedStaticMesh/ISMInstanceDataManager.cpp
+++ b/Engine/Source/Runtime/Engine/Private/InstancedStaticMesh/ISMInstanceDataManager.cpp
@@ -2,6 +2,7 @@
 
 #include "InstancedStaticMesh/ISMInstanceDataManager.h"
 #include "InstancedStaticMesh/ISMInstanceUpdateChangeSet.h"
+#include "InstanceData/InstanceDataHelpers.h"
 #include "InstanceData/InstanceDataUpdateUtils.h"
 
 #include "Engine/InstancedStaticMesh.h"
@@ -1066,11 +1067,9 @@ void FPrimitiveInstanceDataManager::BeginCacheForCookedPlatformData(const ITarge
 
 FISMPrecomputedSpatialHashData FPrimitiveInstanceDataManager::PrecomputeOptimizationData(FInstanceUpdateComponentDesc &&ComponentData, TStridedView<FMatrix> InstanceTransforms)
 {
-	FSpatialHashSortBuilder SortBuilder;
 	const float LocalAbsMaxDisplacement = FMath::Max(-ComponentData.PrimitiveMaterialDesc.MinMaxMaterialDisplacement.X, ComponentData.PrimitiveMaterialDesc.MinMaxMaterialDisplacement.Y)
 									+ ComponentData.PrimitiveMaterialDesc.MaxWorldPositionOffsetDisplacement;
 
-
 	const FVector3f PadExtent = FISMCInstanceDataSceneProxy::GetLocalBoundsPadExtent(ComponentData.PrimitiveLocalToWorld, LocalAbsMaxDisplacement);
 	FRenderBounds InstanceLocalBounds = ComponentData.StaticMeshBounds;
 	InstanceLocalBounds.Min -= PadExtent;
@@ -1078,63 +1077,17 @@ FISMPrecomputedSpatialHashData FPrimitiveInstanceDataManager::PrecomputeOptimiza
 
 	FSphere LocalInstanceSphere = InstanceLocalBounds.ToBoxSphereBounds().GetSphere();
 
-	int32 MinLevel = 0;
-	static const auto CVarInstanceHierarchyMinCellSize = IConsoleManager::Get().FindTConsoleVariableDataFloat(TEXT("r.SceneCulling.MinCellSize"));
-	if (CVarInstanceHierarchyMinCellSize)
-	{
-		// TODO: only one code path to compute this value!!!
-		MinLevel = RenderingSpatialHash::CalcLevel(CVarInstanceHierarchyMinCellSize->GetValueOnAnyThread() - 1.0);
-	}
-
-	SortBuilder.BuildOptimizedSpatialHashOrder(InstanceTransforms.Num(), MinLevel, 
+	FInstanceDataHelpers::FSpatialHashResult SpatialHashResult = FInstanceDataHelpers::BuildSpatialHashData(InstanceTransforms.Num(),
 		[&](int32 InstanceIndex) -> FSphere
 		{
 			FMatrix InstanceLocalToWorld = InstanceTransforms[InstanceIndex] * ComponentData.PrimitiveLocalToWorld;
 			return LocalInstanceSphere.TransformBy(InstanceLocalToWorld);
-		}				
+		}
 	);
 
 	FISMPrecomputedSpatialHashData Result;
-
-	// Pack down the spatial hashes & index remap
-	Result.ProxyIndexToComponentIndexRemap.SetNumUninitialized(InstanceTransforms.Num());
-
-	FInstanceSceneDataBuffers::FCompressedSpatialHashItem CurrentItem;
-	CurrentItem.NumInstances = 0;
-
-	bool bIsIdentityIndexMap = true;
-
-	for (int32 InstanceIndex = 0; InstanceIndex < SortBuilder.SortedInstances.Num(); ++InstanceIndex)
-	{
-		int32 ComponentInstanceIndex = SortBuilder.SortedInstances[InstanceIndex].InstanceIndex;
-		bIsIdentityIndexMap = bIsIdentityIndexMap && InstanceIndex == ComponentInstanceIndex;
-		Result.ProxyIndexToComponentIndexRemap[InstanceIndex] = ComponentInstanceIndex;
-
-		bool bSameLoc = CurrentItem.NumInstances > 0 && CurrentItem.Location == SortBuilder.SortedInstances[InstanceIndex].InstanceLoc;
-		if (bSameLoc)
-		{
-			CurrentItem.NumInstances += 1;
-		}
-		else
-		{
-			if (CurrentItem.NumInstances > 0)
-			{
-				Result.Hashes.Add(CurrentItem);
-			}
-			CurrentItem.Location = SortBuilder.SortedInstances[InstanceIndex].InstanceLoc;
-			CurrentItem.NumInstances = 1;
-		}
-	}
-	if (CurrentItem.NumInstances > 0)
-	{
-		Result.Hashes.Add(CurrentItem);
-	}
-
-	// Don't store a 1:1 mapping
-	if (bIsIdentityIndexMap)
-	{
-		Result.ProxyIndexToComponentIndexRemap.Reset();
-	}
+	Result.Hashes = MoveTemp(SpatialHashResult.SpatialHashes);
+	Result.ProxyIndexToComponentIndexRemap = MoveTemp(SpatialHashResult.ReorderTable);
 	Result.PrimitiveLocalToWorldHash = HashLocalToWorld(ComponentData.PrimitiveLocalToWorld);
 	return Result;
 }
diff --git a/Engine/Source/Runtime/Engine/Public/InstanceDataSceneProxy.h b/Engine/Source/Runtime/Engine/Public/InstanceDataSceneProxy.h
--- a/Engine/Source/Runtime/Engine/Public/InstanceDataSceneProxy.h
+++ b/Engine/Source/Runtime/Engine/Public/InstanceDataSceneProxy.h
@@ -220,6 +220,13 @@ public:
 		// TODO: compress the location - store delta to primitive loc for example
 		RenderingSpatialHash::FLocation64 Location;
 		int32 NumInstances;
+
+		friend inline FArchive& operator<<(FArchive& Ar, FCompressedSpatialHashItem& Item)
+		{
+			Ar << Item.Location;
+			Ar << Item.NumInstances;
+			return Ar;
+		}
 	};
 
 	struct FWriteView
diff --git a/Engine/Source/Runtime/Engine/Public/Rendering/RenderingSpatialHash.h b/Engine/Source/Runtime/Engine/Public/Rendering/RenderingSpatialHash.h
--- a/Engine/Source/Runtime/Engine/Public/Rendering/RenderingSpatialHash.h
+++ b/Engine/Source/Runtime/Engine/Public/Rendering/RenderingSpatialHash.h
@@ -70,6 +70,13 @@ struct TLocation
 		Result.Level = Level;
 		return Result;
 	}
+
+	friend inline FArchive& operator<<(FArchive& Ar, TLocation& Item)
+	{
+		Ar << Item.Coord;
+		Ar << Item.Level;
+		return Ar;
+	}
 };
 
 using FLocation64 = TLocation<int64>;
diff --git a/Engine/Source/Runtime/Mass/MassEngine/Private/Mesh/MassISMRenderStateHelper.cpp b/Engine/Source/Runtime/Mass/MassEngine/Private/Mesh/MassISMRenderStateHelper.cpp
--- a/Engine/Source/Runtime/Mass/MassEngine/Private/Mesh/MassISMRenderStateHelper.cpp
+++ b/Engine/Source/Runtime/Mass/MassEngine/Private/Mesh/MassISMRenderStateHelper.cpp
@@ -2,6 +2,7 @@
 
 #include "MassISMRenderStateHelper.h"
 
+#include "InstanceData/InstanceDataHelpers.h"
 #include "InstancedStaticMeshComponentHelper.h"
 #include "MassEntityManager.h"
 #include "Mesh/MassEngineMeshFragments.h"
@@ -71,22 +72,8 @@ TSharedPtr<FInstanceDataSceneProxy, ESPMode::ThreadSafe>& FMassISMRenderStateHel
 	View.Flags.bHasPerInstanceRandom = PrimitiveMaterialDesc.bAnyMaterialHasPerInstanceRandom && (RenderISMFragment.PerInstanceSMData.Num() > 0);
 	if (View.Flags.bHasPerInstanceRandom)
 	{
-		InstanceRandomIDs.Reset(RenderISMFragment.PerInstanceSMData.Num());
-		InstanceRandomIDs.SetNum(RenderISMFragment.PerInstanceSMData.Num());
-		FRandomStream RandomStream = FRandomStream(RenderISMFragment.InstancingRandomSeed);
-		auto AdditionalRandomSeedsIt = RenderISMFragment.AdditionalRandomSeeds.CreateConstIterator();
-		int32 SeedResetIndex = AdditionalRandomSeedsIt ? AdditionalRandomSeedsIt->StartInstanceIndex : INDEX_NONE;
-		for (int32 Index = 0; Index < InstanceRandomIDs.Num(); ++Index)
-		{
-			// Reset the random stream if necessary
-			if (Index == SeedResetIndex)
-			{
-				RandomStream = FRandomStream(AdditionalRandomSeedsIt->RandomSeed);
-				++AdditionalRandomSeedsIt;
-				SeedResetIndex = AdditionalRandomSeedsIt ? AdditionalRandomSeedsIt->StartInstanceIndex : INDEX_NONE;
-			}
-			InstanceRandomIDs[Index] = RandomStream.GetFraction();
-		}
+		InstanceRandomIDs.SetNumUninitialized(RenderISMFragment.PerInstanceSMData.Num());
+		FInstanceDataHelpers::GenerateInstanceRandomIDs(RenderISMFragment.InstancingRandomSeed, RenderISMFragment.AdditionalRandomSeeds, InstanceRandomIDs);
 	}
 	View.InstanceRandomIDs = InstanceRandomIDs;
 
