[FastGeo] Register streaming delegates for all game worlds
- Removed IsPartitionedWorld guard from PostInitialize delegate registration.
- A non-partitioned world can host Level Instances referencing partitioned maps with FastGeo content whose streaming cells load into the outer world.
- The outer world needs the extension delegates to properly manage FastGeo container registration and unregistration during AddToWorld/RemoveFromWorld.
- Without this fix, FastGeo content in Level Instances hosted by non-partitioned maps would fail to register entirely, and teardown could leave dangling async physics jobs.

#Synced-CL 51624647

diff --git a/Engine/Plugins/Experimental/FastGeoStreaming/Source/FastGeoStreaming/Private/FastGeoWorldSubsystem.cpp b/Engine/Plugins/Experimental/FastGeoStreaming/Source/FastGeoStreaming/Private/FastGeoWorldSubsystem.cpp
--- a/Engine/Plugins/Experimental/FastGeoStreaming/Source/FastGeoStreaming/Private/FastGeoWorldSubsystem.cpp
+++ b/Engine/Plugins/Experimental/FastGeoStreaming/Source/FastGeoStreaming/Private/FastGeoWorldSubsystem.cpp
@@ -186,7 +186,10 @@ void UFastGeoWorldSubsystem::PostInitialize()
 	FWorldDelegates::OnPreRecreateScene.AddUObject(this, &UFastGeoWorldSubsystem::OnPreRecreateScene);
 #endif
 
-	if (World->IsPartitionedWorld() && World->IsGameWorld())
+	// Register streaming delegates for all game worlds, not just partitioned ones.
+	// A non-partitioned world can host Level Instances referencing partitioned maps
+	// with FastGeo content. The streaming cells load into the outer world's level list.
+	if (World->IsGameWorld())
 	{
 		World->OnAllLevelsChanged().AddUObject(this, &UFastGeoWorldSubsystem::OnUpdateLevelStreaming);
 		World->OnAddLevelToWorldExtension().AddUObject(this, &UFastGeoWorldSubsystem::OnAddLevelToWorldExtension);
@@ -201,7 +204,10 @@ void UFastGeoWorldSubsystem::PostInitialize()
 		Handle_OnLevelBeginAddToWorld = FLevelStreamingDelegates::OnLevelBeginMakingVisible.AddUObject(this, &UFastGeoWorldSubsystem::OnLevelStartedAddToWorld);
 		Handle_OnLevelBeginRemoveFromWorld = FLevelStreamingDelegates::OnLevelBeginMakingInvisible.AddUObject(this, &UFastGeoWorldSubsystem::OnLevelStartedRemoveFromWorld);
 
-		Handle_OnForEachHLODObjectInCell = World->GetSubsystem<UWorldPartitionHLODRuntimeSubsystem>()->GetForEachHLODObjectInCellEvent().AddUObject(this, &UFastGeoWorldSubsystem::ForEachHLODObjectInCell);
+		if (UWorldPartitionHLODRuntimeSubsystem* HLODSubsystem = World->GetSubsystem<UWorldPartitionHLODRuntimeSubsystem>())
+		{
+			Handle_OnForEachHLODObjectInCell = HLODSubsystem->GetForEachHLODObjectInCellEvent().AddUObject(this, &UFastGeoWorldSubsystem::ForEachHLODObjectInCell);
+		}
 	}
 }
 
@@ -236,7 +242,10 @@ void UFastGeoWorldSubsystem::Deinitialize()
 
 		if (Handle_OnForEachHLODObjectInCell.IsValid())
 		{
-			World->GetSubsystem<UWorldPartitionHLODRuntimeSubsystem>()->GetForEachHLODObjectInCellEvent().Remove(Handle_OnForEachHLODObjectInCell);
+			if (UWorldPartitionHLODRuntimeSubsystem* HLODSubsystem = World->GetSubsystem<UWorldPartitionHLODRuntimeSubsystem>())
+			{
+				HLODSubsystem->GetForEachHLODObjectInCellEvent().Remove(Handle_OnForEachHLODObjectInCell);
+			}
 			Handle_OnForEachHLODObjectInCell.Reset();
 		}
 	}
