[FastGeo] Updated comments about async render state destruction not being throttled.

#Synced-CL 51624647

diff --git a/Engine/Plugins/Experimental/FastGeoStreaming/Source/FastGeoStreaming/Private/FastGeoContainer.cpp b/Engine/Plugins/Experimental/FastGeoStreaming/Source/FastGeoStreaming/Private/FastGeoContainer.cpp
--- a/Engine/Plugins/Experimental/FastGeoStreaming/Source/FastGeoStreaming/Private/FastGeoContainer.cpp
+++ b/Engine/Plugins/Experimental/FastGeoStreaming/Source/FastGeoStreaming/Private/FastGeoContainer.cpp
@@ -1041,7 +1041,7 @@ void UFastGeoContainer::OnDestroyRenderStateBegin_GameThread()
 
 	PendingDestroy.bIsInBlockingWait = GetWorldSubsystem()->IsWaitingForCompletion();
 
-	// TODO_FASTGEO: Evaluate if we need throttling of the async destruction task
+	// There is no throttling for the async destruction task
 	PendingDestroy.RenderState.NumToProcess = PendingDestroy.RenderState.ComponentsToProcess.Num();
 }
 
@@ -1110,7 +1110,6 @@ void UFastGeoContainer::OnDestroyRenderState_Concurrent()
 {
 	TRACE_CPUPROFILER_EVENT_SCOPE(UFastGeoContainer::OnDestroyRenderState_Concurrent);
 
-	// TODO_FASTGEO: Evaluate if we need throttling of the async destruction task
 	check(PendingDestroy.RenderState.NumToProcess == PendingDestroy.RenderState.ComponentsToProcess.Num());
 	{
 		FFastGeoDestroyRenderStateContext Context(GetWorld()->Scene);
