FastGeo - Defer SSAM requirement warning until FastGeo is actually in use
* Move the FSimpleStreamableAssetManager config check out of FFastGeoStreamingModule::StartupModule so it no longer fires at engine boot just because the FastGeo plugin is enabled.
* Instead, warn once per world on the first PushAsyncCreateRenderStateJob call. This ensures the warning only appears when FastGeo containers are actively pushing render work, and dedicated servers that never create render state won't see it.

#Synced-CL 52288206

diff --git a/Engine/Plugins/Experimental/FastGeoStreaming/Source/FastGeoStreaming/Internal/FastGeoWorldSubsystem.h b/Engine/Plugins/Experimental/FastGeoStreaming/Source/FastGeoStreaming/Internal/FastGeoWorldSubsystem.h
--- a/Engine/Plugins/Experimental/FastGeoStreaming/Source/FastGeoStreaming/Internal/FastGeoWorldSubsystem.h
+++ b/Engine/Plugins/Experimental/FastGeoStreaming/Source/FastGeoStreaming/Internal/FastGeoWorldSubsystem.h
@@ -140,6 +140,7 @@ private:
 	int32 UsedNumComponentsToProcessBudget = 0;
 	bool bWaitingForCompletion = false;
 	bool bIsReregistering = false;
+	bool bHasWarnedAboutSSAM = false;
 
 	TUniquePtr<FFastGeoAsyncRenderStateJobQueue> AsyncRenderStateJobQueue;
 	TArray<TWeakObjectPtr<UFastGeoContainer>> ReregisteringFastGeoContainers;
diff --git a/Engine/Plugins/Experimental/FastGeoStreaming/Source/FastGeoStreaming/Private/FastGeoStreamingModule.cpp b/Engine/Plugins/Experimental/FastGeoStreaming/Source/FastGeoStreaming/Private/FastGeoStreamingModule.cpp
--- a/Engine/Plugins/Experimental/FastGeoStreaming/Source/FastGeoStreaming/Private/FastGeoStreamingModule.cpp
+++ b/Engine/Plugins/Experimental/FastGeoStreaming/Source/FastGeoStreaming/Private/FastGeoStreamingModule.cpp
@@ -38,13 +38,6 @@ bool FFastGeoStreamingModule::IsAsyncRenderWorkAllowed()
 
 void FFastGeoStreamingModule::StartupModule()
 {
-#if !WITH_EDITOR
-	if (IConsoleVariable* SSAMCVar = IConsoleManager::Get().FindConsoleVariable(TEXT("s.StreamableAssets.UseSimpleStreamableAssetManager")))
-	{
-		ensureMsgf(SSAMCVar->GetBool(), TEXT("FSimpleStreamableAssetManager must be enabled for FastGeo. Set s.StreamableAssets.UseSimpleStreamableAssetManager=1"));
-	}
-#endif
-
 	IPrimitiveComponent::AddImplementer(TComponentInterfaceImplementation<IPrimitiveComponent>(UFastGeoContainer::StaticClass(),  [](UObject* Obj)
 	{
 		return CastChecked<UFastGeoContainer>(Obj)->GetPrimitiveComponents();
diff --git a/Engine/Plugins/Experimental/FastGeoStreaming/Source/FastGeoStreaming/Private/FastGeoWorldSubsystem.cpp b/Engine/Plugins/Experimental/FastGeoStreaming/Source/FastGeoStreaming/Private/FastGeoWorldSubsystem.cpp
--- a/Engine/Plugins/Experimental/FastGeoStreaming/Source/FastGeoStreaming/Private/FastGeoWorldSubsystem.cpp
+++ b/Engine/Plugins/Experimental/FastGeoStreaming/Source/FastGeoStreaming/Private/FastGeoWorldSubsystem.cpp
@@ -2,6 +2,7 @@
 
 #include "FastGeoWorldSubsystem.h"
 #include "FastGeoContainer.h"
+#include "FastGeoLog.h"
 #include "FastGeoPrimitiveComponent.h"
 #include "FastGeoStreamingModule.h"
 #include "FastGeoSurrogateActor.h"
@@ -17,6 +18,7 @@
 #include "GameFramework/ActorPrimitiveColorHandler.h"
 #include "Engine/GameViewportClient.h"
 #include "ProfilingDebugging/CsvProfiler.h"
+#include "Streaming/SimpleStreamableAssetManager.h"
 
 #define LOCTEXT_NAMESPACE "FastGeoWorldSubsystem"
 
@@ -752,6 +754,13 @@ void UFastGeoWorldSubsystem::PushAsyncCreateRenderStateJob(UFastGeoContainer* Fa
 {
 	check(IsInGameThread());
 	check(FastGeo->GetWorld() == GetWorld());
+
+	if (!bHasWarnedAboutSSAM && !FSimpleStreamableAssetManager::IsEnabled())
+	{
+		bHasWarnedAboutSSAM = true;
+		UE_LOG(LogFastGeoStreaming, Warning, TEXT("FSimpleStreamableAssetManager must be enabled for FastGeo. Set s.StreamableAssets.UseSimpleStreamableAssetManager=1"));
+	}
+
 	GetOrCreateJobQueue()->AddJob({ FastGeo, FFastGeoAsyncRenderStateJobQueue::EJobType::CreateRenderState });
 }
 
