[FastGeo] Add overlay materials to ForEachMaterial
- Add component, per-slot, and asset-level overlay materials to ForEachMaterial for static mesh and skinned mesh components.
- Ensures overlay materials are included in MaterialAssetHardRefs for GC keepalive and ConditionalPostLoad pre-caching.

#Synced-CL 52208678

diff --git a/Engine/Plugins/Experimental/FastGeoStreaming/Source/FastGeoStreaming/Private/FastGeoSkinnedMeshComponent.cpp b/Engine/Plugins/Experimental/FastGeoStreaming/Source/FastGeoStreaming/Private/FastGeoSkinnedMeshComponent.cpp
--- a/Engine/Plugins/Experimental/FastGeoStreaming/Source/FastGeoStreaming/Private/FastGeoSkinnedMeshComponent.cpp
+++ b/Engine/Plugins/Experimental/FastGeoStreaming/Source/FastGeoStreaming/Private/FastGeoSkinnedMeshComponent.cpp
@@ -233,6 +233,26 @@ void FFastGeoSkinnedMeshComponentBase::ForEachMaterial(TFunctionRef<void(UMateri
 			{
 				AddMaterial(SlotMat.MaterialInterface);
 			}
+			if (SlotMat.OverlayMaterialInterface)
+			{
+				AddMaterial(SlotMat.OverlayMaterialInterface);
+			}
+		}
+	}
+
+	// Component overlay material
+	const FSkinnedMeshSceneProxyDesc& Desc = GetSkinnedMeshSceneProxyDesc();
+	if (Desc.OverlayMaterial)
+	{
+		AddMaterial(Desc.OverlayMaterial);
+	}
+
+	// Component per-slot overlay materials
+	for (const TObjectPtr<UMaterialInterface>& Mat : Desc.MaterialSlotsOverlayMaterial)
+	{
+		if (Mat)
+		{
+			AddMaterial(Mat);
 		}
 	}
 }
diff --git a/Engine/Plugins/Experimental/FastGeoStreaming/Source/FastGeoStreaming/Private/FastGeoStaticMeshComponent.cpp b/Engine/Plugins/Experimental/FastGeoStreaming/Source/FastGeoStreaming/Private/FastGeoStaticMeshComponent.cpp
--- a/Engine/Plugins/Experimental/FastGeoStreaming/Source/FastGeoStreaming/Private/FastGeoStaticMeshComponent.cpp
+++ b/Engine/Plugins/Experimental/FastGeoStreaming/Source/FastGeoStreaming/Private/FastGeoStaticMeshComponent.cpp
@@ -216,6 +216,26 @@ void FFastGeoStaticMeshComponentBase::ForEachMaterial(TFunctionRef<void(UMateria
 			{
 				AddMaterial(SlotMat.MaterialInterface);
 			}
+			if (SlotMat.OverlayMaterialInterface)
+			{
+				AddMaterial(SlotMat.OverlayMaterialInterface);
+			}
+		}
+	}
+
+	// Component overlay material
+	const FStaticMeshSceneProxyDesc& Desc = GetStaticMeshSceneProxyDesc();
+	if (Desc.OverlayMaterial)
+	{
+		AddMaterial(Desc.OverlayMaterial);
+	}
+
+	// Component per-slot overlay materials
+	for (const TObjectPtr<UMaterialInterface>& Mat : Desc.MaterialSlotsOverlayMaterial)
+	{
+		if (Mat)
+		{
+			AddMaterial(Mat);
 		}
 	}
 }
