[FastGeo] Optimize PSO precaching for deferred create and fallback material recreate
- Optimize deferred create (DelayUntilPSOPrecached): components whose PSOs are not ready during initial create are parked and created later via the async job queue when PSO compilation completes.
- Optimize fallback material recreate (UseFallbackMaterialUntilPSOPrecached): components created with fallback material are recreated with real material via async job queue once PSOs are ready.
- Both strategies now use budgeted, time-sliced processing, eliminating 20-30ms GT spikes when many components complete PSO compilation simultaneously.
- Sync-mode implementations added for both strategies using budgeted ParallelFor (AdvanceRenderStateBudgeted).
- Extract TFastGeoRenderStateBatch and AdvanceRenderStateBudgeted into shared FastGeoRenderStateUtils.h for reuse across create, deferred create, and recreate paths.
- Add DeferredCreateRenderState and RecreateRenderState job types to the async render state job queue (async mode).
- Add ContainersWithPSODeferredWork safety net polling in WorldSubsystem::Tick to detect lost PSO callbacks and drive deferred create after OnAddLevelToWorldExtension completes.
- Add ContainersWithAsyncPSORecreateWork tracking in WorldSubsystem::Tick to drive async recreate jobs for level-streamed containers after OnAddLevelToWorldExtension completes (async mode).
- Add FFastGeoPendingRecreateManager with two-pass processing: immediate non-PSO recreates and budgeted PSO recreates.
- Add sync-mode time budget CVar (FastGeo.SyncPSORecreateRenderState.TimeBudgetMS) for PSO recreate render state work.
- Handle BlockTillLevelStreamingCompleted for PSO-deferred components by force-creating proxies when GT is blocked.
- Handle in-flight recreate jobs during container unregistration with async-safe abort and cleanup.
- Add bCreatedWithPSOFallbackMaterial flag and NeedsPSORecreate() virtual on FFastGeoPrimitiveComponent.
- Simplify PSO callback routing through OnComponentPSOPrecacheCompleted for all component types.

#Synced-CL 52208678

diff --git a/Engine/Plugins/Experimental/FastGeoStreaming/Source/FastGeoStreaming/Internal/FastGeoAsyncRenderStateJobQueue.h b/Engine/Plugins/Experimental/FastGeoStreaming/Source/FastGeoStreaming/Internal/FastGeoAsyncRenderStateJobQueue.h
--- a/Engine/Plugins/Experimental/FastGeoStreaming/Source/FastGeoStreaming/Internal/FastGeoAsyncRenderStateJobQueue.h
+++ b/Engine/Plugins/Experimental/FastGeoStreaming/Source/FastGeoStreaming/Internal/FastGeoAsyncRenderStateJobQueue.h
@@ -17,7 +17,11 @@ struct FFastGeoAsyncRenderStateJobQueue
 	{
 		PrecachePSOs,
 		CreateRenderState,
-		DestroyRenderState
+		DestroyRenderState,
+#if !WITH_EDITOR && UE_WITH_PSO_PRECACHING
+		DeferredCreateRenderState,
+		RecreateRenderState,
+#endif
 	};
 
 	struct FJob
diff --git a/Engine/Plugins/Experimental/FastGeoStreaming/Source/FastGeoStreaming/Internal/FastGeoComponent.h b/Engine/Plugins/Experimental/FastGeoStreaming/Source/FastGeoStreaming/Internal/FastGeoComponent.h
--- a/Engine/Plugins/Experimental/FastGeoStreaming/Source/FastGeoStreaming/Internal/FastGeoComponent.h
+++ b/Engine/Plugins/Experimental/FastGeoStreaming/Source/FastGeoStreaming/Internal/FastGeoComponent.h
@@ -95,6 +95,8 @@ public:
 	virtual void PrecachePSOs_Concurrent() {}
 	virtual FPSOPrecacheComponentData& GetPSOPrecacheComponentData() = 0;
 	virtual const FPSOPrecacheComponentData& GetPSOPrecacheComponentData() const = 0;
+	// Returns true if the proxy was created with a fallback material and needs recreation when PSO completes.
+	virtual bool NeedsPSORecreate() const { return false; }
 #endif
 
 protected:
diff --git a/Engine/Plugins/Experimental/FastGeoStreaming/Source/FastGeoStreaming/Internal/FastGeoContainer.h b/Engine/Plugins/Experimental/FastGeoStreaming/Source/FastGeoStreaming/Internal/FastGeoContainer.h
--- a/Engine/Plugins/Experimental/FastGeoStreaming/Source/FastGeoStreaming/Internal/FastGeoContainer.h
+++ b/Engine/Plugins/Experimental/FastGeoStreaming/Source/FastGeoStreaming/Internal/FastGeoContainer.h
@@ -11,6 +11,7 @@
 #include "FastGeoComponentCluster.h"
 #include "FastGeoHLOD.h"
 #include "Containers/BitArray.h"
+#include "FastGeoRenderStateUtils.h"
 #include "UObject/ObjectPtr.h"
 #include "UObject/SoftObjectPtr.h"
 
@@ -103,37 +104,10 @@ private:
 
 struct FFastGeoComponentClusterPendingState
 {
-	struct FRenderState
-	{
-		FRenderState()
-		{
-			Reset();
-		}
-
-		void Reset()
-		{
-			ComponentsToProcess.Reset();
-			NumToProcess = 0;
-			NumProcessed = 0;
-			TotalNumProcessed = 0;
-		}
-
-		bool IsCompleted() const
-		{
-			return TotalNumProcessed >= ComponentsToProcess.Num();
-		}
-
-		TArray<FFastGeoComponent*> ComponentsToProcess;
-
-		int32 NumToProcess;
-		int32 NumProcessed;
-		int32 TotalNumProcessed;
-	};
-
-	FRenderState RenderState;
+	FFastGeoRenderStateBatch RenderState;
 
 #if !WITH_EDITOR && UE_WITH_PSO_PRECACHING
-	FRenderState PrecachePSOState;
+	FFastGeoRenderStateBatch PrecachePSOState;
 #endif
 
 	struct FPhysicsState
@@ -158,8 +132,6 @@ struct FFastGeoComponentClusterPendingState
 		std::atomic<int32> TotalNumProcessed;
 	} PhysicsState;
 
-	bool bIsInBlockingWait;
-	
 	FFastGeoComponentClusterPendingState()
 	{
 		Reset();
@@ -172,7 +144,6 @@ struct FFastGeoComponentClusterPendingState
 		PrecachePSOState.Reset();
 #endif
 		PhysicsState.Reset();
-		bIsInBlockingWait = false;
 	}
 
 	bool HasAnyPendingState(bool bIncludePSOPrecaching = true) const
@@ -230,6 +201,9 @@ public:
 	bool IsFullyUnregistered() const;
 	void Tick(bool bWaitForCompletion = false);
 	bool HasAnyPendingTasks(bool bIncludePSOPrecaching = true) const;
+	// Returns true if registration-related tasks (create/destroy) are pending.
+	// Excludes post-registration PSO recreate work (UseFallbackMaterialUntilPSOPrecached).
+	bool HasAnyPendingRegistrationTasks(bool bIncludePSOPrecaching = true) const;
 	bool HasAnyPendingCreateTasks(bool bIncludePSOPrecaching = true) const;
 	bool HasAnyPendingDestroyTasks() const;
 	void PrecachePSOs();
@@ -267,6 +241,40 @@ public:
 	void OnCreateRenderStateBegin_GameThread();
 	void OnCreateRenderState_Concurrent();
 	void OnCreateRenderStateEnd_GameThread();
+
+#if !WITH_EDITOR && UE_WITH_PSO_PRECACHING
+	// Called from PSO precache completion callback. Routes to the appropriate path:
+	// - DelayUntilPSOPrecached: moves component to ready set for deferred create.
+	// - UseFallbackMaterialUntilPSOPrecached: pushes to recreate queue for material swap.
+	// - Otherwise: no-op (early completion detected by CheckPSODeferredReadiness polling).
+	void OnComponentPSOPrecacheCompleted(FFastGeoComponent* Component);
+
+	// DelayUntilPSOPrecached: deferred proxy creation for components whose PSOs were not ready.
+	void OnPSODeferredReady(FFastGeoComponent* Component);
+	void OnDeferredCreateRenderStateBegin_GameThread();
+	void OnDeferredCreateRenderState_Concurrent();
+	void OnDeferredCreateRenderStateEnd_GameThread();
+
+	// Checks if deferred components are ready for proxy creation and triggers a deferred create pass.
+	// bPollPending: when true, polls IsPSOPrecaching() on pending components to detect lost PSO callbacks (safety net).
+	// Called with bPollPending=true from:
+	// - WorldSubsystem::Tick (periodic safety net poll)
+	// - OnCreateRenderStateEnd_GameThread (detect callbacks that fired before parking)
+	// - OnDeferredCreateRenderStateEnd_GameThread (re-check after batch completes)
+	// Called with bPollPending=false from:
+	// - OnPSODeferredReady (per-callback, skips O(n) poll).
+	void CheckPSODeferredReadiness(bool bPollPending = false);
+
+	// UseFallbackMaterialUntilPSOPrecached: recreate proxy with real material once PSOs are ready.
+	void RoutePSORecreateComponent(FFastGeoComponent* Component);
+	void PushPSORecreateComponent(FFastGeoComponent* Component);
+	void OnRecreateRenderStateBegin_GameThread();
+	void OnRecreateRenderState_Concurrent();
+	void OnRecreateRenderStateEnd_GameThread();
+
+	bool HasPSODeferredComponents() const { return !PSODeferredPendingComponents.IsEmpty() || !PSODeferredReadyComponents.IsEmpty(); }
+	bool HasAnyPendingRecreateTasks() const;
+#endif
 	
 	void OnDestroyRenderStateBegin_GameThread();
 	void OnDestroyRenderState_Concurrent();
@@ -324,6 +332,30 @@ protected:
 	FFastGeoComponentClusterPendingState PendingCreate;
 	FFastGeoComponentClusterPendingState PendingDestroy;
 
+#if !WITH_EDITOR && UE_WITH_PSO_PRECACHING
+	// DelayUntilPSOPrecached: components whose proxy creation was deferred because PSOs were not ready.
+	// Part of registration -- blocks AddToWorld until all deferred proxies are created.
+	// Components start in PSODeferredPendingComponents. As each PSO callback fires, the
+	// component moves to PSODeferredReadyComponents. When the pending set is empty (or the
+	// safety poll detects all are ready), the ready array becomes ComponentsToProcess for
+	// the deferred create pass.
+	TSet<FFastGeoComponent*> PSODeferredPendingComponents;
+	TArray<FFastGeoComponent*> PSODeferredReadyComponents;
+	FFastGeoRenderStateBatch PSODeferredCreateState;
+	bool bIsProcessingPSODeferredCreate = false;
+
+	// UseFallbackMaterialUntilPSOPrecached: components created with fallback material that
+	// need recreation with the real material once PSOs are ready.
+	// Post-registration -- does NOT block AddToWorld.
+	// Populated by PushPSORecreateComponent (async-only, asserted), processed by the async
+	// job queue via StartRecreateRenderStateWork -> PushAsyncRecreateRenderStateJob.
+	// Sync mode does not use these fields -- it routes through MarkRenderStateDirty ->
+	// ProcessPendingRecreate instead.
+	TArray<FFastGeoComponent*> PSORecreateComponents;
+	FFastGeoRenderStateBatch PSORecreateState;
+	bool bIsProcessingPSORecreate = false;
+#endif
+
 private:
 
 	UFastGeoWorldSubsystem* GetWorldSubsystem() const;
@@ -347,13 +379,29 @@ private:
 	void StartPrecachePSOWork();
 	void StartCreateRenderStateWork();
 	void StartDestroyRenderStateWork();
+	template <bool bDestroyFirst>
+	void ExecuteCreateRenderState_Concurrent(FFastGeoRenderStateBatch& State);
+#if !WITH_EDITOR && UE_WITH_PSO_PRECACHING
+	void ForceCreateDeferredComponents();
+	void StartDeferredCreateRenderStateWork();
+	void StartRecreateRenderStateWork();
+	bool TickDeferredCreate_Sync(const UE::FTimeout& Timeout);
+#endif
 	void TickRenderWork(bool bWaitForCompletion);
 	void TickRenderWork_Sync(bool bWaitForCompletion);
 	bool TickPrecachePSO_Sync(const UE::FTimeout& Timeout);
 	bool TickCreateRenderState_Sync(const UE::FTimeout& Timeout);
 	bool TickDestroyRenderState_Sync(const UE::FTimeout& Timeout);
 	bool IsPrecachePSOCompleted() const;
-	bool IsCreateRenderStateCompleted() const;
+
+	// Returns true when all components in the current create pass have been attempted.
+	// Components whose proxy was deferred (e.g., DelayUntilPSOPrecached) are counted as attempted.
+	// Use IsCreateRenderStateFullyCompleted() to check if all proxies actually exist.
+	bool IsInitialCreateRenderStatePassCompleted() const;
+
+	// Returns true when all render state creation is truly done, including PSO-deferred proxies.
+	bool IsCreateRenderStateFullyCompleted() const;
+
 	bool IsDestroyRenderStateCompleted() const;
 
 	void CollectAssetReferences();
@@ -368,6 +416,8 @@ private:
 	void UnregisterBodyInstances();
 	void OnRegisterCompleted();
 	void OnUnregisterCompleted();
+	void TryCompleteRegistration();
+	void TryCompleteUnregistration();
 	void ApplyRegistrationTargetState();
 	void StartRegisterTransition();
 	void StartUnregisterTransition();
@@ -492,7 +542,7 @@ private:
 	TObjectPtr<AFastGeoSurrogateActor> SurrogateActor;
 
 #if !WITH_EDITOR && UE_WITH_PSO_PRECACHING
-	bool bProcessedPSOPrecaching = false;
+	bool bPSOPrecachingInitiated = false;
 #endif
 
 	// Resolved material cache (transient).
diff --git a/Engine/Plugins/Experimental/FastGeoStreaming/Source/FastGeoStreaming/Internal/FastGeoPrimitiveComponent.h b/Engine/Plugins/Experimental/FastGeoStreaming/Source/FastGeoStreaming/Internal/FastGeoPrimitiveComponent.h
--- a/Engine/Plugins/Experimental/FastGeoStreaming/Source/FastGeoStreaming/Internal/FastGeoPrimitiveComponent.h
+++ b/Engine/Plugins/Experimental/FastGeoStreaming/Source/FastGeoStreaming/Internal/FastGeoPrimitiveComponent.h
@@ -118,6 +118,7 @@ public:
 	virtual void PrecachePSOs_Concurrent() override;
 	virtual FPSOPrecacheComponentData& GetPSOPrecacheComponentData() override { return PSOPrecacheComponentData; }
 	virtual const FPSOPrecacheComponentData& GetPSOPrecacheComponentData() const override { return PSOPrecacheComponentData; }
+	virtual bool NeedsPSORecreate() const override { return bCreatedWithPSOFallbackMaterial; }
 #endif
 	bool CheckPSOPrecachingAndBoostPriority(EPSOPrecachePriority NewPSOPrecachePriority = EPSOPrecachePriority::High);
 	//~End FFastGeoComponent interface
@@ -199,6 +200,10 @@ protected:
 
 	/** All PSO related data for the component */
 	FPSOPrecacheComponentData PSOPrecacheComponentData;
+
+	// True when the proxy was created with a fallback material (UseFallbackMaterialUntilPSOPrecached).
+	// Used by OnComponentPSOPrecacheCompleted to know this component needs a recreate when PSO completes.
+	bool bCreatedWithPSOFallbackMaterial = false;
 #endif
 	UFastGeoSurrogateComponent* GetSurrogateComponent() const;
 
diff --git a/Engine/Plugins/Experimental/FastGeoStreaming/Source/FastGeoStreaming/Internal/FastGeoRenderStateUtils.h b/Engine/Plugins/Experimental/FastGeoStreaming/Source/FastGeoStreaming/Internal/FastGeoRenderStateUtils.h
new file mode 100644
--- /dev/null
+++ b/Engine/Plugins/Experimental/FastGeoStreaming/Source/FastGeoStreaming/Internal/FastGeoRenderStateUtils.h
@@ -0,0 +1,106 @@
+// Copyright Epic Games, Inc. All Rights Reserved.
+
+#pragma once
+
+#include "Containers/Array.h"
+#include "Async/ParallelFor.h"
+#include "Async/TaskGraphInterfaces.h"
+#include "Misc/Timeout.h"
+
+class FFastGeoComponent;
+
+// General-purpose progress tracker for budgeted render state work.
+// Used by create, deferred create, recreate, and ProcessPendingRecreate paths.
+// Default element type is FFastGeoComponent* (raw pointer). ProcessPendingRecreate uses
+// TFastGeoRenderStateBatch<FFastGeoRegisteredComponent> to avoid persisting raw pointers
+// across frames -- stale components are validated inline via TryGetRegistered.
+template <typename T = FFastGeoComponent*>
+struct TFastGeoRenderStateBatch
+{
+	TFastGeoRenderStateBatch()
+	{
+		Reset();
+	}
+
+	void Reset()
+	{
+		ComponentsToProcess.Reset();
+		NumToProcess = 0;
+		NumProcessed = 0;
+		TotalNumProcessed = 0;
+		bIsInBlockingWait = false;
+	}
+
+	bool IsCompleted() const
+	{
+		return TotalNumProcessed >= ComponentsToProcess.Num();
+	}
+
+	TArray<T> ComponentsToProcess;
+
+	int32 NumToProcess;
+	int32 NumProcessed;
+	int32 TotalNumProcessed;
+	bool bIsInBlockingWait = false;
+};
+
+using FFastGeoRenderStateBatch = TFastGeoRenderStateBatch<FFastGeoComponent*>;
+
+namespace FastGeo
+{
+	// Processes components in a batch using ParallelFor with FTaskTagScope (ParallelGT pattern).
+	// Used by sync-mode create, destroy, deferred create, recreate, and ProcessPendingRecreate.
+	// Returns true if all components were processed, false if timeout expired.
+	// WorkFn receives T& (e.g., FFastGeoComponent*& for raw batches, FFastGeoRegisteredComponent& for safe batches).
+	template <typename T, typename FWorkFn>
+	FORCEINLINE bool AdvanceRenderStateBudgeted(TFastGeoRenderStateBatch<T>& State, const UE::FTimeout& Timeout, FWorkFn&& WorkFn)
+	{
+		check(IsInGameThread());
+
+		if (State.IsCompleted())
+		{
+			return true;
+		}
+
+		if (Timeout.IsExpired())
+		{
+			return false;
+		}
+
+		std::atomic<int32> NextIndex{ 0 };
+		std::atomic<int32> NumProcessed{ 0 };
+		static constexpr int32 MinElementsPerTask = 8;
+		const int32 NumRemaining = State.NumToProcess;
+		const int32 StartIndex = State.TotalNumProcessed;
+		const int32 NumWorkers = FTaskGraphInterface::Get().GetNumWorkerThreads();
+		const int32 NumTasks = FMath::Clamp(FMath::DivideAndRoundUp(NumRemaining, MinElementsPerTask), 1, NumWorkers);
+
+		ParallelFor(NumTasks, [&](int32 TaskIndex)
+		{
+			FTaskTagScope Scope(ETaskTag::EParallelGameThread);
+			while (true)
+			{
+				if (Timeout.IsExpired())
+				{
+					return;
+				}
+
+				const int32 LocalIndex = NextIndex.fetch_add(1, std::memory_order_relaxed);
+				if (LocalIndex >= NumRemaining)
+				{
+					return;
+				}
+
+				T& Component = State.ComponentsToProcess[StartIndex + LocalIndex];
+				WorkFn(Component);
+
+				NumProcessed.fetch_add(1, std::memory_order_relaxed);
+			}
+		});
+
+		State.NumProcessed = NumProcessed.load(std::memory_order_relaxed);
+		State.TotalNumProcessed += State.NumProcessed;
+		State.NumToProcess = State.ComponentsToProcess.Num() - State.TotalNumProcessed;
+		return State.IsCompleted();
+	}
+}
diff --git a/Engine/Plugins/Experimental/FastGeoStreaming/Source/FastGeoStreaming/Internal/FastGeoWorldSubsystem.h b/Engine/Plugins/Experimental/FastGeoStreaming/Source/FastGeoStreaming/Internal/FastGeoWorldSubsystem.h
--- a/Engine/Plugins/Experimental/FastGeoStreaming/Source/FastGeoStreaming/Internal/FastGeoWorldSubsystem.h
+++ b/Engine/Plugins/Experimental/FastGeoStreaming/Source/FastGeoStreaming/Internal/FastGeoWorldSubsystem.h
@@ -38,7 +38,7 @@ public:
 	virtual TStatId GetStatId() const override;
 	//~ End UTickableWorldSubsystem interface
 
-	void AddToComponentsPendingRecreate(FFastGeoComponent* InComponentPendingRecreate);
+	void AddToComponentsPendingRecreate(FFastGeoComponent* InComponent);
 	void ProcessPendingRecreate();
 
 	/**
@@ -59,6 +59,14 @@ public:
 	 */
 	void PushAsyncDestroyRenderStateJob(UFastGeoContainer* FastGeo);
 
+#if !WITH_EDITOR && UE_WITH_PSO_PRECACHING
+	/** Start a job to create render state for PSO-deferred components (DelayUntilPSOPrecached). */
+	void PushAsyncDeferredCreateJob(UFastGeoContainer* FastGeo);
+
+	/** Start a job to recreate render state for components with fallback material (UseFallbackMaterialUntilPSOPrecached). */
+	void PushAsyncRecreateRenderStateJob(UFastGeoContainer* FastGeo);
+#endif
+
 	/**
 	 * Update progress of asynchronous render state creation and destruction.
 	 * @param bWaitForCompletion : Whether to wait for pending asynchronous tasks to complete.
@@ -110,7 +118,7 @@ private:
 #if DO_CHECK
 	void OnLevelAddedToWorld(ULevel* Level, UWorld* World);
 	void OnLevelRemovedFromWorld(ULevel* Level, UWorld* World);
-	void CheckNoPendingTasks(ULevel* Level, UWorld* World);
+	void CheckNoPendingTasks(ULevel* Level, UWorld* World, bool bIncludeRecreateTasks = true);
 #endif
 #if WITH_EDITOR
 	void OnPreRecreateScene(UWorld* World);
@@ -123,6 +131,12 @@ private:
 	void PushRuntimeContainerForRegistration(UFastGeoContainer* FastGeo);
 	void PushRuntimeContainerForDestruction(UFastGeoContainer* FastGeo);
 	void DrainPendingRuntimeContainers();
+#if !WITH_EDITOR && UE_WITH_PSO_PRECACHING
+	void RegisterContainerWithPSODeferredWork(UFastGeoContainer* FastGeo);
+	void UnregisterContainerWithPSODeferredWork(UFastGeoContainer* FastGeo);
+	void RegisterContainerWithAsyncPSORecreateWork(UFastGeoContainer* FastGeo);
+	void UnregisterContainerWithAsyncPSORecreateWork(UFastGeoContainer* FastGeo);
+#endif
 
 	UPROPERTY()
 	TSet<TObjectPtr<UFastGeoContainer>> PendingRuntimeRegisterContainers;
@@ -135,7 +149,41 @@ private:
 	FDelegateHandle Handle_OnLevelBeginRemoveFromWorld;
 	FDelegateHandle Handle_OnForEachHLODObjectInCell;
 
-	TArray<FFastGeoRegisteredComponent> ComponentsPendingRecreate;
+	// Manages components pending render state recreation (e.g., MarkRenderStateDirty from runtime
+	// property changes, decal fade, or PSO fallback material swap in sync mode).
+	// Non-PSO recreates are processed immediately (same frame guarantee).
+	// PSO recreates (NeedsPSORecreate) are time-sliced with a per-frame budget.
+	struct FFastGeoPendingRecreateManager
+	{
+		void Add(FFastGeoComponent* Component);
+		void Process(UWorld* World);
+		bool IsEmpty() const;
+		void Reset();
+
+#if WITH_EDITOR
+		const TArray<FFastGeoRegisteredComponent>& GetComponents() const { return Components; }
+#endif
+
+	private:
+		TArray<FFastGeoRegisteredComponent> Components;
+#if !WITH_EDITOR && UE_WITH_PSO_PRECACHING
+		TFastGeoRenderStateBatch<FFastGeoRegisteredComponent> PSORecreateBatch;
+#endif
+		bool bHasNonPSORecreate = false;
+	} PendingRecreateManager;
+
+#if !WITH_EDITOR && UE_WITH_PSO_PRECACHING
+	// Containers with PSO-deferred components waiting for PSO callbacks.
+	// Polled each Tick() to detect lost callbacks (component destroyed, world torn down, etc.)
+	// and unblock registration that would otherwise hang forever.
+	TSet<TWeakObjectPtr<UFastGeoContainer>> ContainersWithPSODeferredWork;
+
+	// Containers with pending PSO recreate work (UseFallbackMaterialUntilPSOPrecached).
+	// Tracked so Tick() can drive ProcessAsyncRenderStateJobs for recreate jobs that were
+	// pushed after OnAddLevelToWorldExtension completed. See Tick() comment for details.
+	TSet<TWeakObjectPtr<UFastGeoContainer>> ContainersWithAsyncPSORecreateWork;
+#endif
+
 	static bool bEnableDebugView;
 
 	FRWLock Lock;
diff --git a/Engine/Plugins/Experimental/FastGeoStreaming/Source/FastGeoStreaming/Private/FastGeoAsyncRenderStateJobQueue.cpp b/Engine/Plugins/Experimental/FastGeoStreaming/Source/FastGeoStreaming/Private/FastGeoAsyncRenderStateJobQueue.cpp
--- a/Engine/Plugins/Experimental/FastGeoStreaming/Source/FastGeoStreaming/Private/FastGeoAsyncRenderStateJobQueue.cpp
+++ b/Engine/Plugins/Experimental/FastGeoStreaming/Source/FastGeoStreaming/Private/FastGeoAsyncRenderStateJobQueue.cpp
@@ -46,6 +46,14 @@ void FFastGeoAsyncRenderStateJobQueue::AddJob(const FJob& Job)
 		case EJobType::DestroyRenderState:
 			Job.FastGeo->OnDestroyRenderStateBegin_GameThread();
 			break;
+#if !WITH_EDITOR && UE_WITH_PSO_PRECACHING
+		case EJobType::DeferredCreateRenderState:
+			Job.FastGeo->OnDeferredCreateRenderStateBegin_GameThread();
+			break;
+		case EJobType::RecreateRenderState:
+			Job.FastGeo->OnRecreateRenderStateBegin_GameThread();
+			break;
+#endif
 		}
 	}
 }
@@ -96,7 +104,6 @@ void FFastGeoAsyncRenderStateJobQueue::Launch()
 void FFastGeoAsyncRenderStateJobQueue::OnUpdateLevelStreamingDone()
 {
 	TriggerIsReadyToRunAsyncTasksEvent();
-
 }
 
 void FFastGeoAsyncRenderStateJobQueue::TriggerIsReadyToRunAsyncTasksEvent()
@@ -198,6 +205,20 @@ void FFastGeoAsyncRenderStateJobQueue::FJob::Execute() const
 			FastGeo->OnDestroyRenderState_Concurrent();
 		}
 		break;
+#if !WITH_EDITOR && UE_WITH_PSO_PRECACHING
+		case EJobType::DeferredCreateRenderState:
+		{
+			TRACE_CPUPROFILER_EVENT_SCOPE(FFastGeoAsyncRenderStateJobQueue::DeferredCreateRenderState);
+			FastGeo->OnDeferredCreateRenderState_Concurrent();
+		}
+		break;
+		case EJobType::RecreateRenderState:
+		{
+			TRACE_CPUPROFILER_EVENT_SCOPE(FFastGeoAsyncRenderStateJobQueue::RecreateRenderState);
+			FastGeo->OnRecreateRenderState_Concurrent();
+		}
+		break;
+#endif
 		}
 	}
 }
@@ -217,6 +238,14 @@ void FFastGeoAsyncRenderStateJobQueue::FJob::OnPostExecute_GameThread() const
 		case EJobType::DestroyRenderState:
 			FastGeo->OnDestroyRenderStateEnd_GameThread();
 			break;
+#if !WITH_EDITOR && UE_WITH_PSO_PRECACHING
+		case EJobType::DeferredCreateRenderState:
+			FastGeo->OnDeferredCreateRenderStateEnd_GameThread();
+			break;
+		case EJobType::RecreateRenderState:
+			FastGeo->OnRecreateRenderStateEnd_GameThread();
+			break;
+#endif
 		}
 	}
 }
\ No newline at end of file
diff --git a/Engine/Plugins/Experimental/FastGeoStreaming/Source/FastGeoStreaming/Private/FastGeoContainer.cpp b/Engine/Plugins/Experimental/FastGeoStreaming/Source/FastGeoStreaming/Private/FastGeoContainer.cpp
--- a/Engine/Plugins/Experimental/FastGeoStreaming/Source/FastGeoStreaming/Private/FastGeoContainer.cpp
+++ b/Engine/Plugins/Experimental/FastGeoStreaming/Source/FastGeoStreaming/Private/FastGeoContainer.cpp
@@ -180,6 +180,15 @@ namespace FastGeo::Private
 		FGuid TaskId;
 		TArray<TObjectPtr<UObject>> Assets;
 	};
+
+	// Sync mode calls End_GameThread once after all budget-sliced batches complete.
+	// Reset the scan range to [0, Num) so End_GameThread processes all components,
+	// not just the last batch (which would miss components delayed in earlier batches).
+	static void ResetBatchForSyncCompletion(FFastGeoRenderStateBatch& State)
+	{
+		State.TotalNumProcessed = 0;
+		State.NumProcessed = State.ComponentsToProcess.Num();
+	}
 }
 
 void UFastGeoContainer::ApplyRegistrationTargetState()
@@ -236,7 +245,8 @@ void UFastGeoContainer::Register()
 	RegistrationTargetState = ERegistrationTargetState::Registered;
 	ApplyRegistrationTargetState();
 
-	if (HasAnyPendingTasks())
+	// Recreate (UseFallbackMaterialUntilPSOPrecached) is post-registration and shouldn't prevent re-registration setup.
+	if (HasAnyPendingRegistrationTasks())
 	{
 		Tick();
 	}
@@ -257,8 +267,9 @@ void UFastGeoContainer::StartRegisterTransition()
 	// Surrogate flag may be stale during FGlobalComponentReregisterContext (no actor-level callbacks).
 	check(!IsUsingSurrogateComponents() || !SurrogateActor->HasRegisteredWithFastGeo() || GetWorldSubsystem()->IsReregistering());
 
-	constexpr bool bIncludePSOPrecaching = false; // PSO precaching is allowed to run before registration
-	if (const bool bHasAnyPendingTasks = HasAnyPendingTasks(bIncludePSOPrecaching))
+	// PSO precaching is allowed to run before registration. Recreate from a prior registration cycle may still be in flight -- that's OK.
+	constexpr bool bIncludePSOPrecaching = false;
+	if (const bool bHasAnyPendingTasks = HasAnyPendingRegistrationTasks(bIncludePSOPrecaching))
 	{
 		checkf(!bHasAnyPendingTasks, TEXT("UFastGeoContainer::StartRegisterTransition() was called while the container still has pending tasks. This will lead to invalid behavior."));
 		return;
@@ -301,14 +312,14 @@ void UFastGeoContainer::StartRegisterTransition()
 #if !WITH_EDITOR && UE_WITH_PSO_PRECACHING
 	// Prepare PendingCreate.PSOPrecachingState
 	PendingCreate.PrecachePSOState.Reset();
-	if (!FFastGeoStreamingModule::IsAsyncRenderWorkAllowed() && !bProcessedPSOPrecaching)
+	if (!FFastGeoStreamingModule::IsAsyncRenderWorkAllowed() && !bPSOPrecachingInitiated)
 	{
 		if (IsComponentPSOPrecachingEnabled())
 		{
 			PendingCreate.PrecachePSOState.ComponentsToProcess = PendingCreate.RenderState.ComponentsToProcess;
 			PendingCreate.PrecachePSOState.NumToProcess = PendingCreate.PrecachePSOState.ComponentsToProcess.Num();
 		}
-		bProcessedPSOPrecaching = (PendingCreate.PrecachePSOState.ComponentsToProcess.Num() == 0);
+		bPSOPrecachingInitiated = (PendingCreate.PrecachePSOState.ComponentsToProcess.Num() == 0);
 	}
 #endif
 
@@ -317,10 +328,7 @@ void UFastGeoContainer::StartRegisterTransition()
 	StartCreatePhysicsStateWork();
 
 	// If no work needed for any system, registration is already complete
-	if (!HasAnyPendingCreateTasks())
-	{
-		OnRegisterCompleted();
-	}
+	TryCompleteRegistration();
 
 	const UNavigationSystemBase* NavigationSystem = GetWorld()->GetNavigationSystem();
 	if (NavigationSystem && NavigationSystem->IsWorldInitDone())
@@ -345,7 +353,8 @@ bool UFastGeoContainer::IsRegistered() const
 
 bool UFastGeoContainer::IsFullyRegistered() const
 {
-	return (RegistrationState == ERegistrationState::Registered) && !HasAnyPendingTasks();
+	// Recreate is post-registration background work -- container IS fully registered.
+	return (RegistrationState == ERegistrationState::Registered) && !HasAnyPendingRegistrationTasks();
 }
 
 bool UFastGeoContainer::IsUnregistering() const
@@ -385,6 +394,9 @@ void UFastGeoContainer::OnUnregisterCompleted()
 	TRACE_CPUPROFILER_EVENT_SCOPE(UFastGeoContainer::OnUnregisterCompleted);
 	// Verify that unregistration is truly complete
 	check(!HasAnyPendingDestroyTasks());
+#if !WITH_EDITOR && UE_WITH_PSO_PRECACHING
+	check(!HasAnyPendingRecreateTasks());
+#endif
 	check(RegistrationState == ERegistrationState::Unregistering);
 	RegistrationState = ERegistrationState::Unregistered;
 
@@ -395,6 +407,27 @@ void UFastGeoContainer::OnUnregisterCompleted()
 	ApplyRegistrationTargetState();
 }
 
+void UFastGeoContainer::TryCompleteRegistration()
+{
+	if (RegistrationState == ERegistrationState::Registering && !HasAnyPendingCreateTasks())
+	{
+		OnRegisterCompleted();
+	}
+}
+
+void UFastGeoContainer::TryCompleteUnregistration()
+{
+	if (RegistrationState == ERegistrationState::Unregistering
+		&& !HasAnyPendingDestroyTasks()
+#if !WITH_EDITOR && UE_WITH_PSO_PRECACHING
+		&& !HasAnyPendingRecreateTasks()
+#endif
+		)
+	{
+		OnUnregisterCompleted();
+	}
+}
+
 FFastGeoCreateRuntimeResult UFastGeoContainer::CreateRuntime(UWorld* InWorld, FName InDebugName, TFunctionRef<void(FFastGeoComponentCluster&)> InInitCluster, bool bInCollectReferences)
 {
 	TRACE_CPUPROFILER_EVENT_SCOPE(UFastGeoContainer::CreateRuntime);
@@ -532,6 +565,34 @@ void UFastGeoContainer::StartUnregisterTransition()
 
 	PendingDestroy.Reset();
 
+#if !WITH_EDITOR && UE_WITH_PSO_PRECACHING
+	// Clean up any PSO deferred/recreate state from the registration that just completed.
+	// Without this, stale state would cause HasAnyPendingCreateTasks/HasAnyPendingRecreateTasks
+	// to return true on re-registration (e.g., FGlobalComponentReregisterContext), triggering asserts.
+
+	// Deferred create is part of registration and should be complete by now.
+	check(PSODeferredPendingComponents.IsEmpty());
+	check(PSODeferredReadyComponents.IsEmpty());
+	check(PSODeferredCreateState.ComponentsToProcess.IsEmpty());
+	PSODeferredPendingComponents.Reset();
+	PSODeferredReadyComponents.Reset();
+	PSODeferredCreateState.Reset();
+	GetWorldSubsystem()->UnregisterContainerWithPSODeferredWork(this);
+
+	// If a recreate job is in flight, don't touch PSORecreateState (worker may be reading it).
+	// OnRecreateRenderStateEnd_GameThread will detect RegistrationState == Unregistering and abort,
+	// and will call UnregisterContainerWithAsyncPSORecreateWork at that point.
+	// bIsProcessingPSORecreate stays true so HasAnyPendingRecreateTasks() returns true until
+	// the in-flight job completes and cleans up.
+	if (!bIsProcessingPSORecreate)
+	{
+		check(PSORecreateState.ComponentsToProcess.IsEmpty());
+		PSORecreateState.Reset();
+		GetWorldSubsystem()->UnregisterContainerWithAsyncPSORecreateWork(this);
+	}
+	PSORecreateComponents.Reset();
+#endif
+
 	ForEachComponentCluster([this](FFastGeoComponentCluster& ComponentCluster)
 	{
 		ComponentCluster.OnUnregister();
@@ -548,87 +609,29 @@ void UFastGeoContainer::StartUnregisterTransition()
 	StartDestroyPhysicsStateWork();
 
 	// If no work needed for any system, unregistration is already complete
-	if (!HasAnyPendingDestroyTasks())
-	{
-		OnUnregisterCompleted();
-	}
-}
-
-namespace FastGeo::Private
-{
-	template <typename FWorkFn>
-	FORCEINLINE bool AdvanceRenderStateBudgeted(FFastGeoComponentClusterPendingState::FRenderState& State, const UE::FTimeout& Timeout, FWorkFn&& WorkFn)
-	{
-		check(IsInGameThread());
-
-		if (State.IsCompleted())
-		{
-			return true;
-		}
-
-		if (Timeout.IsExpired())
-		{
-			return false;
-		}
-
-		std::atomic<int32> NextIndex{ 0 };
-		std::atomic<int32> NumProcessed{ 0 };
-		static constexpr int32 MinElementsPerTask = 8;
-		const int32 NumRemaining = State.NumToProcess;
-		const int32 StartIndex = State.TotalNumProcessed;
-		const int32 NumWorkers = FTaskGraphInterface::Get().GetNumWorkerThreads();
-		const int32 NumTasks = FMath::Clamp(FMath::DivideAndRoundUp(NumRemaining, MinElementsPerTask), 1, NumWorkers);
-
-		ParallelFor(NumTasks, [&](int32 TaskIndex)
-		{
-			FTaskTagScope Scope(ETaskTag::EParallelGameThread);
-			while (true)
-			{
-				if (Timeout.IsExpired())
-				{
-					return;
-				}
-
-				const int32 LocalIndex = NextIndex.fetch_add(1, std::memory_order_relaxed);
-				if (LocalIndex >= NumRemaining)
-				{
-					return;
-				}
-
-				FFastGeoComponent* Component = State.ComponentsToProcess[StartIndex + LocalIndex];
-				WorkFn(Component);
-
-				NumProcessed.fetch_add(1, std::memory_order_relaxed);
-			}
-		});
-
-		State.NumProcessed = NumProcessed.load(std::memory_order_relaxed);
-		State.TotalNumProcessed += State.NumProcessed;
-		State.NumToProcess = State.ComponentsToProcess.Num() - State.TotalNumProcessed;
-		return State.IsCompleted();
-	}
+	TryCompleteUnregistration();
 }
 
 bool UFastGeoContainer::TickPrecachePSO_Sync(const UE::FTimeout& Timeout)
 {
 #if !WITH_EDITOR && UE_WITH_PSO_PRECACHING
-	if (!bProcessedPSOPrecaching)
+	if (!bPSOPrecachingInitiated)
 	{
 		TRACE_CPUPROFILER_EVENT_SCOPE(UFastGeoContainer::TickPrecachePSO_Sync);
 		check(!FFastGeoStreamingModule::IsAsyncRenderWorkAllowed());
 
-		bProcessedPSOPrecaching = FastGeo::Private::AdvanceRenderStateBudgeted(PendingCreate.PrecachePSOState, Timeout, [](FFastGeoComponent* Component)
+		bPSOPrecachingInitiated = FastGeo::AdvanceRenderStateBudgeted(PendingCreate.PrecachePSOState, Timeout, [](FFastGeoComponent* Component)
 		{
-			TRACE_CPUPROFILER_EVENT_SCOPE(UFastGeoContainer::PrecachePSOs);
+			TRACE_CPUPROFILER_EVENT_SCOPE(FFastGeoComponent_PrecachePSOs);
 			Component->PrecachePSOs_Concurrent();
 		});
 
-		if (bProcessedPSOPrecaching)
+		if (bPSOPrecachingInitiated)
 		{
 			PendingCreate.PrecachePSOState.Reset();
 		}
 	}
-	return bProcessedPSOPrecaching;
+	return bPSOPrecachingInitiated;
 #else
 	return true;
 #endif
@@ -645,16 +648,15 @@ bool UFastGeoContainer::TickCreateRenderState_Sync(const UE::FTimeout& Timeout)
 		OnCreateRenderStateBegin_GameThread();
 	}
 
-	const bool bCompleted = FastGeo::Private::AdvanceRenderStateBudgeted(PendingCreate.RenderState, Timeout, [](FFastGeoComponent* Component)
+	const bool bCompleted = FastGeo::AdvanceRenderStateBudgeted(PendingCreate.RenderState, Timeout, [](FFastGeoComponent* Component)
 	{
-		TRACE_CPUPROFILER_EVENT_SCOPE(UFastGeoContainer::CreateRenderState);
+		TRACE_CPUPROFILER_EVENT_SCOPE(FFastGeoComponent_CreateRenderState);
 		Component->CreateRenderState(nullptr);
 	});
 
 	if (bCompleted)
 	{
-		// Force NumProcessed to zero because OnCreateRenderStateEnd_GameThread was designed to increment TotalNumProcessed
-		PendingCreate.RenderState.NumProcessed = 0;
+		FastGeo::Private::ResetBatchForSyncCompletion(PendingCreate.RenderState);
 		OnCreateRenderStateEnd_GameThread();
 	}
 
@@ -672,7 +674,7 @@ bool UFastGeoContainer::TickDestroyRenderState_Sync(const UE::FTimeout& Timeout)
 	check(IsInGameThread());
 	check(!FFastGeoStreamingModule::IsAsyncRenderWorkAllowed());
 
-	FFastGeoComponentClusterPendingState::FRenderState& State = PendingDestroy.RenderState;
+	FFastGeoRenderStateBatch& State = PendingDestroy.RenderState;
 	if (State.ComponentsToProcess.Num() > 0 && State.NumToProcess == 0)
 	{
 		OnDestroyRenderStateBegin_GameThread();
@@ -738,9 +740,9 @@ void UFastGeoContainer::TickPhysicsWork(bool bWaitForCompletion)
 void UFastGeoContainer::StartPrecachePSOWork()
 {
 #if !WITH_EDITOR && UE_WITH_PSO_PRECACHING
-	if (FFastGeoStreamingModule::IsAsyncRenderWorkAllowed() && !bProcessedPSOPrecaching)
+	if (FFastGeoStreamingModule::IsAsyncRenderWorkAllowed() && !bPSOPrecachingInitiated)
 	{
-		bProcessedPSOPrecaching = true;
+		bPSOPrecachingInitiated = true;
 		if (IsComponentPSOPrecachingEnabled())
 		{
 			UFastGeoWorldSubsystem* WorldSubsystem = GetWorldSubsystem();
@@ -752,14 +754,14 @@ void UFastGeoContainer::StartPrecachePSOWork()
 
 void UFastGeoContainer::StartCreateRenderStateWork()
 {
-	if (FFastGeoStreamingModule::IsAsyncRenderWorkAllowed() && !IsCreateRenderStateCompleted())
+	if (FFastGeoStreamingModule::IsAsyncRenderWorkAllowed() && !IsInitialCreateRenderStatePassCompleted())
 	{
 #if !WITH_EDITOR && UE_WITH_PSO_PRECACHING
 		// PSO precaching should have been started by UFastGeoWorldSubsystem::OnLevelStreamingStateChanged
-		if (!ensure(bProcessedPSOPrecaching))
+		if (!ensure(bPSOPrecachingInitiated))
 		{
 			StartPrecachePSOWork();
-			check(bProcessedPSOPrecaching);
+			check(bPSOPrecachingInitiated);
 		}
 #endif
 		UFastGeoWorldSubsystem* WorldSubsystem = GetWorldSubsystem();
@@ -779,17 +781,37 @@ void UFastGeoContainer::StartDestroyRenderStateWork()
 bool UFastGeoContainer::IsPrecachePSOCompleted() const
 {
 #if !WITH_EDITOR && UE_WITH_PSO_PRECACHING
-	return bProcessedPSOPrecaching && PendingCreate.PrecachePSOState.IsCompleted();
+	return bPSOPrecachingInitiated && PendingCreate.PrecachePSOState.IsCompleted();
 #else
 	return true;
 #endif
 }
 
-bool UFastGeoContainer::IsCreateRenderStateCompleted() const
+bool UFastGeoContainer::IsInitialCreateRenderStatePassCompleted() const
 {
 	return PendingCreate.RenderState.IsCompleted();
 }
 
+bool UFastGeoContainer::IsCreateRenderStateFullyCompleted() const
+{
+	if (!IsInitialCreateRenderStatePassCompleted())
+	{
+		return false;
+	}
+#if !WITH_EDITOR && UE_WITH_PSO_PRECACHING
+	if (HasPSODeferredComponents())
+	{
+		return false;
+	}
+	if (bIsProcessingPSODeferredCreate)
+	{
+		return false;
+	}
+	check(PSODeferredCreateState.IsCompleted());
+#endif
+	return true;
+}
+
 bool UFastGeoContainer::IsDestroyRenderStateCompleted() const
 {
 	return PendingDestroy.RenderState.IsCompleted();
@@ -822,7 +844,7 @@ void UFastGeoContainer::TickRenderWork_Sync(bool bWaitForCompletion)
 	{
 	case ERegistrationState::Registering:
 		{
-			if (!IsPrecachePSOCompleted() || !IsCreateRenderStateCompleted())
+			if (!IsPrecachePSOCompleted() || !IsInitialCreateRenderStatePassCompleted())
 			{
 				const UE::FTimeout Timeout = bWaitForCompletion ? UE::FTimeout::Never() : GetWorld()->GetAddToWorldTimeout();
 				if (TickPrecachePSO_Sync(Timeout))
@@ -830,10 +852,19 @@ void UFastGeoContainer::TickRenderWork_Sync(bool bWaitForCompletion)
 					check(IsPrecachePSOCompleted());
 					if (TickCreateRenderState_Sync(Timeout))
 					{
-						check(IsCreateRenderStateCompleted());
+						check(IsInitialCreateRenderStatePassCompleted());
 					}
 				}
 			}
+#if !WITH_EDITOR && UE_WITH_PSO_PRECACHING
+			// Deferred create: components whose PSOs were not ready during initial create.
+			// Uses AddToWorld timeout since deferred create is part of registration.
+			if (IsInitialCreateRenderStatePassCompleted() && !PSODeferredCreateState.ComponentsToProcess.IsEmpty())
+			{
+				const UE::FTimeout Timeout = bWaitForCompletion ? UE::FTimeout::Never() : GetWorld()->GetAddToWorldTimeout();
+				TickDeferredCreate_Sync(Timeout);
+			}
+#endif
 		}
 		break;
 	case ERegistrationState::Unregistering:
@@ -849,9 +880,13 @@ void UFastGeoContainer::TickRenderWork_Sync(bool bWaitForCompletion)
 		}
 		break;
 	default:
-		check(!HasAnyPendingTasks());
+		// Recreate is valid in the default state (post-registration) and handled below.
+		check(!HasAnyPendingRegistrationTasks(/*bIncludePSOPrecaching=*/ false));
 		break;
 	}
+
+	// Note: sync-mode strategy 2 recreate (UseFallbackMaterialUntilPSOPrecached) is handled by MarkRenderStateDirty -> ProcessPendingRecreate, not here.
+	// PushPSORecreateComponent is async-only (asserted). See OnComponentPSOPrecacheCompleted for the routing.
 }
 
 void UFastGeoContainer::Tick(bool bWaitForCompletion)
@@ -862,24 +897,48 @@ void UFastGeoContainer::Tick(bool bWaitForCompletion)
 		return;
 	}
 
+	// Callers may pass bWaitForCompletion=false even when GT is blocked (e.g., AddToWorld called from BlockTillLevelStreamingCompleted).
+	// Detect the blocking state from the world to ensure we force-create PSO-deferred proxies instead of parking indefinitely.
+	const bool bShouldWaitForCompletion = bWaitForCompletion || GetWorld()->GetIsInBlockTillLevelStreamingCompleted();
+
 	if (!HasAnyPendingTasks())
 	{
 		ApplyRegistrationTargetState();
 	}
 
+#if !WITH_EDITOR && UE_WITH_PSO_PRECACHING
+	// If we enter a blocking wait while already parked with PSO-deferred components, un-park and force-create.
+	// GT is blocked (e.g., BlockTillLevelStreamingCompleted) and cannot pump the task graph, so PSO callbacks will never fire.
+	if (bShouldWaitForCompletion && HasPSODeferredComponents())
+	{
+		// Wait for any in-flight deferred create job to complete before touching PSODeferredCreateState (a worker thread could be reading it).
+		TickRenderWork(bShouldWaitForCompletion);
+		ForceCreateDeferredComponents();
+	}
+#endif
+
 	do
 	{
-		TickRenderWork(bWaitForCompletion);
-		TickPhysicsWork(bWaitForCompletion);
+		TickRenderWork(bShouldWaitForCompletion);
+		TickPhysicsWork(bShouldWaitForCompletion);
 
-		if (bWaitForCompletion && !HasAnyPendingTasks())
+		if (bShouldWaitForCompletion && !HasAnyPendingTasks())
 		{
 			ApplyRegistrationTargetState();
 		}
-	} while (bWaitForCompletion && HasAnyPendingTasks());
+	} while (bShouldWaitForCompletion && HasAnyPendingTasks());
 }
 
 bool UFastGeoContainer::HasAnyPendingTasks(bool bIncludePSOPrecaching) const
+{
+	return HasAnyPendingCreateTasks(bIncludePSOPrecaching) || HasAnyPendingDestroyTasks()
+#if !WITH_EDITOR && UE_WITH_PSO_PRECACHING
+		|| HasAnyPendingRecreateTasks()
+#endif
+		;
+}
+
+bool UFastGeoContainer::HasAnyPendingRegistrationTasks(bool bIncludePSOPrecaching) const
 {
 	return HasAnyPendingCreateTasks(bIncludePSOPrecaching) || HasAnyPendingDestroyTasks();
 }
@@ -887,8 +946,15 @@ bool UFastGeoContainer::HasAnyPendingTasks(bool bIncludePSOPrecaching) const
 bool UFastGeoContainer::HasAnyPendingCreateTasks(bool bIncludePSOPrecaching) const
 {
 	// Physics state can have no PendingState but OnAsyncCreatePhysicsStateEnd_GameThread has not been called yet
-	return (PhysicsState == EPhysicsState::Creating) || PendingCreate.HasAnyPendingState(bIncludePSOPrecaching);
+	return (PhysicsState == EPhysicsState::Creating) || PendingCreate.HasAnyPendingState(bIncludePSOPrecaching) || !IsCreateRenderStateFullyCompleted();
+}
+
+#if !WITH_EDITOR && UE_WITH_PSO_PRECACHING
+bool UFastGeoContainer::HasAnyPendingRecreateTasks() const
+{
+	return bIsProcessingPSORecreate || !PSORecreateComponents.IsEmpty();
 }
+#endif
 
 bool UFastGeoContainer::HasAnyPendingDestroyTasks() const
 {
@@ -1066,12 +1132,11 @@ void UFastGeoContainer::OnCreateRenderStateBegin_GameThread()
 	check(PendingCreate.RenderState.NumToProcess == 0);
 	check(PendingCreate.RenderState.TotalNumProcessed != PendingCreate.RenderState.ComponentsToProcess.Num());
 
-	PendingCreate.bIsInBlockingWait = GetWorldSubsystem()->IsWaitingForCompletion();
+	PendingCreate.RenderState.bIsInBlockingWait = GetWorldSubsystem()->IsWaitingForCompletion();
 	PendingCreate.RenderState.NumToProcess = PendingCreate.RenderState.ComponentsToProcess.Num() - PendingCreate.RenderState.TotalNumProcessed;
 
 	// Refresh cached time (see FFastGeoDecalComponent::CreateRenderState)
 	AsyncContext.UpdateTimeSeconds(GetWorld()->GetTimeSeconds());
-
 }
 
 void UFastGeoContainer::OnDestroyRenderStateBegin_GameThread()
@@ -1079,19 +1144,20 @@ void UFastGeoContainer::OnDestroyRenderStateBegin_GameThread()
 	check(PendingDestroy.RenderState.NumToProcess == 0);
 	check(PendingDestroy.RenderState.TotalNumProcessed != PendingDestroy.RenderState.ComponentsToProcess.Num());
 
-	PendingDestroy.bIsInBlockingWait = GetWorldSubsystem()->IsWaitingForCompletion();
+	PendingDestroy.RenderState.bIsInBlockingWait = GetWorldSubsystem()->IsWaitingForCompletion();
 
 	// There is no throttling for the async destruction task
 	PendingDestroy.RenderState.NumToProcess = PendingDestroy.RenderState.ComponentsToProcess.Num();
 }
 
-void UFastGeoContainer::OnCreateRenderState_Concurrent()
+template <bool bDestroyFirst>
+void UFastGeoContainer::ExecuteCreateRenderState_Concurrent(FFastGeoRenderStateBatch& State)
 {
 	check(FFastGeoStreamingModule::IsAsyncRenderWorkAllowed());
 
 	static constexpr int32 MinNumElementsToProcessPerThread = 8;
-	const int32 NumComponentsToProcess = PendingCreate.RenderState.NumToProcess;
-	const int32 MaxNumThreads = PendingCreate.bIsInBlockingWait ? INT32_MAX : FastGeo::Private::GAsyncRenderStateTaskParallelWorkerCount;
+	const int32 NumComponentsToProcess = State.NumToProcess;
+	const int32 MaxNumThreads = State.bIsInBlockingWait ? INT32_MAX : FastGeo::Private::GAsyncRenderStateTaskParallelWorkerCount;
 	const int32 NumThreads = FMath::Clamp(NumComponentsToProcess / MinNumElementsToProcessPerThread, 1, MaxNumThreads);
 	const bool bIsParallelForAllowed = NumThreads > 1 && FApp::ShouldUseThreadingForPerformance();
 
@@ -1113,7 +1179,7 @@ void UFastGeoContainer::OnCreateRenderState_Concurrent()
 	{
 		ParallelFor(NumThreads, [&](int32 Index)
 		{
-			TRACE_CPUPROFILER_EVENT_SCOPE(UFastGeoContainer::OnCreateRenderState_Concurrent);
+			TRACE_CPUPROFILER_EVENT_SCOPE(UFastGeoContainer::ExecuteCreateRenderState_Concurrent);
 			// Currently necessary for FSimpleStreamableAssetManager::FRegister (to be revisited)
 			FTaskTagScope Scope(ETaskTag::EParallelGameThread);
 
@@ -1132,18 +1198,32 @@ void UFastGeoContainer::OnCreateRenderState_Concurrent()
 					return; 
 				}
 
-				FFastGeoComponent* ComponentToProcess = PendingCreate.RenderState.ComponentsToProcess[PendingCreate.RenderState.TotalNumProcessed + LocalIndex];
-				ComponentToProcess->CreateRenderState(/*Context = */ nullptr);
+				FFastGeoComponent* ComponentToProcess = State.ComponentsToProcess[State.TotalNumProcessed + LocalIndex];
+				if constexpr (bDestroyFirst)
+				{
+					ComponentToProcess->DestroyRenderState(nullptr);
+				}
+				ComponentToProcess->CreateRenderState(nullptr);
 
 				NumProcessed.fetch_add(1, std::memory_order_relaxed);
 			}
 		}, bIsParallelForAllowed ? EParallelForFlags::None : EParallelForFlags::ForceSingleThread);
 	}
 
-	// Adjust NumRenderStatesToProcess to the actual number of states we processed. 
-	PendingCreate.RenderState.NumProcessed = NumProcessed.load(std::memory_order_relaxed);
+	State.NumProcessed = NumProcessed.load(std::memory_order_relaxed);
+	WorldSubsystem->CommitAsyncRenderStateTasksBudget_Concurrent(Timeout.GetElapsedSeconds() * 1000, State.NumProcessed, TimeEpoch);
+}
 
-	WorldSubsystem->CommitAsyncRenderStateTasksBudget_Concurrent(Timeout.GetElapsedSeconds() * 1000, PendingCreate.RenderState.NumProcessed, TimeEpoch);
+// Explicit template instantiations
+template void UFastGeoContainer::ExecuteCreateRenderState_Concurrent<false>(FFastGeoRenderStateBatch&);
+#if !WITH_EDITOR && UE_WITH_PSO_PRECACHING
+template void UFastGeoContainer::ExecuteCreateRenderState_Concurrent<true>(FFastGeoRenderStateBatch&);
+#endif
+
+void UFastGeoContainer::OnCreateRenderState_Concurrent()
+{
+	TRACE_CPUPROFILER_EVENT_SCOPE(UFastGeoContainer::OnCreateRenderState_Concurrent);
+	ExecuteCreateRenderState_Concurrent<false>(PendingCreate.RenderState);
 }
 
 void UFastGeoContainer::OnDestroyRenderState_Concurrent()
@@ -1171,26 +1251,412 @@ void UFastGeoContainer::OnCreateRenderStateEnd_GameThread()
 		check(PendingCreate.RenderState.TotalNumProcessed != PendingCreate.RenderState.ComponentsToProcess.Num());
 	}
 
-	PendingCreate.RenderState.TotalNumProcessed += PendingCreate.RenderState.NumProcessed;
+	const int32 BatchNumProcessed = PendingCreate.RenderState.NumProcessed;
+	const int32 BatchStartIndex = PendingCreate.RenderState.TotalNumProcessed;
+
+	PendingCreate.RenderState.TotalNumProcessed += BatchNumProcessed;
 	PendingCreate.RenderState.NumToProcess = 0;
 	PendingCreate.RenderState.NumProcessed = 0;
-	PendingCreate.bIsInBlockingWait = false;
+	PendingCreate.RenderState.bIsInBlockingWait = false;
 
-	if (!IsCreateRenderStateCompleted())
+#if !WITH_EDITOR && UE_WITH_PSO_PRECACHING
+	// Collect PSO-delayed components from the batch just processed.
+	// Covers both DelayUntilPSOPrecached (all component types delayed) and
+	// UseFallbackMaterialUntilPSOPrecached (decals/lights delayed, static meshes use fallback).
+	// Both strategies check != AlwaysCreate in their CreateSceneProxy.
+	if (GetPSOPrecacheProxyCreationStrategy() != EPSOPrecacheProxyCreationStrategy::AlwaysCreate)
+	{
+		for (int32 i = 0; i < BatchNumProcessed; ++i)
+		{
+			FFastGeoComponent* Comp = PendingCreate.RenderState.ComponentsToProcess[BatchStartIndex + i];
+			if (Comp->IsRenderStateDelayed())
+			{
+				PSODeferredPendingComponents.Add(Comp);
+			}
+			// Safety net for UseFallbackMaterialUntilPSOPrecached: if the PSO callback fired
+			// before CreateRenderState ran, the callback saw NeedsPSORecreate()=false (proxy
+			// didn't exist yet) and was silently ignored. Detect this by checking for components
+			// created with fallback material whose PSO has already completed.
+			else if (Comp->NeedsPSORecreate() && !Comp->GetPSOPrecacheComponentData().IsPSOPrecaching())
+			{
+				RoutePSORecreateComponent(Comp);
+			}
+		}
+	}
+#endif
+
+	if (!IsInitialCreateRenderStatePassCompleted())
 	{
 		check(FFastGeoStreamingModule::IsAsyncRenderWorkAllowed());
 		StartCreateRenderStateWork();
 	}
+#if !WITH_EDITOR && UE_WITH_PSO_PRECACHING
+	else if (HasPSODeferredComponents())
+	{
+		if (GetWorldSubsystem()->IsWaitingForCompletion())
+		{
+			ForceCreateDeferredComponents();
+		}
+		else
+		{
+			// Normal per-frame path: park and wait for PSO callbacks.
+			// As PSO callbacks fire, components move from PSODeferredPendingComponents to PSODeferredReadyComponents.
+			// When the pending set is empty (or the periodic safety poll detects all are ready), 
+			// CheckPSODeferredReadiness() triggers the deferred create pass.
+			PSODeferredReadyComponents.Reserve(PSODeferredPendingComponents.Num());
+			GetWorldSubsystem()->RegisterContainerWithPSODeferredWork(this);
+
+			// Handle the race where PSO callbacks fired before we parked. Those callbacks
+			// did nothing (component wasn't in the pending set yet). Poll IsPSOPrecaching()
+			// now to detect any that already completed and move them to the ready set.
+			CheckPSODeferredReadiness(/*bPollPending=*/ true);
+		}
+	}
+#endif
 	else
 	{
 		// Render work complete - check if all registration systems are done
-		if (!HasAnyPendingCreateTasks())
+		TryCompleteRegistration();
+	}
+}
+
+#if !WITH_EDITOR && UE_WITH_PSO_PRECACHING
+void UFastGeoContainer::OnComponentPSOPrecacheCompleted(FFastGeoComponent* Component)
+{
+	check(IsInGameThread());
+
+	if (PSODeferredPendingComponents.Contains(Component))
+	{
+		// DelayUntilPSOPrecached: component had no proxy, PSO is now ready for deferred create
+		OnPSODeferredReady(Component);
+	}
+	else if (Component->NeedsPSORecreate())
+	{
+		// UseFallbackMaterialUntilPSOPrecached: component has proxy with fallback material, PSO is now ready, swap to real material.
+		RoutePSORecreateComponent(Component);
+	}
+	// Else: component is Delayed but not yet parked (PSO callback fired before OnCreateRenderStateEnd_GameThread collected it), 
+	// or in None state (proxy failed for non-PSO reasons -- retrying won't help). 
+	// CheckPSODeferredReadiness() will detect the early completion via IsPSOPrecaching() polling when End_GameThread parks.
+}
+
+void UFastGeoContainer::OnPSODeferredReady(FFastGeoComponent* Component)
+{
+	check(IsInGameThread());
+
+	// Move from pending to ready. O(1) remove from TSet.
+	PSODeferredPendingComponents.Remove(Component);
+	PSODeferredReadyComponents.Add(Component);
+
+	// Process ready components incrementally -- don't wait for all pending to complete.
+	// CheckPSODeferredReadiness handles early-outs (nothing ready, pass in flight).
+	CheckPSODeferredReadiness();
+}
+
+void UFastGeoContainer::CheckPSODeferredReadiness(bool bPollPending)
+{
+	TRACE_CPUPROFILER_EVENT_SCOPE(UFastGeoContainer::CheckPSODeferredReadiness);
+	check(IsInGameThread());
+
+	if (PSODeferredPendingComponents.IsEmpty() && PSODeferredReadyComponents.IsEmpty())
 	{
-			OnRegisterCompleted();
+		return;
+	}
+
+	if (bPollPending)
+	{
+		// Safety net: detect components whose PSO callback was lost (component destroyed before
+		// callback fired, lifecycle ID mismatch in FPSOPrecacheFinishedTask, etc.). Without this
+		// check, a lost callback would leave the container parked forever with AddToWorld never
+		// completing. We poll IsPSOPrecaching() to detect PSOs that completed without notifying us.
+		// Skipped when called per-callback from OnPSODeferredReady (bPollPending=false) to avoid
+		// O(n) iteration on every callback. See header comment for full list of call sites.
+		for (auto It = PSODeferredPendingComponents.CreateIterator(); It; ++It)
+		{
+			FFastGeoComponent* Comp = *It;
+			if (!Comp->GetPSOPrecacheComponentData().IsPSOPrecaching())
+			{
+				PSODeferredReadyComponents.Add(Comp);
+				It.RemoveCurrent();
+			}
+		}
+	}
+
+	if (PSODeferredReadyComponents.IsEmpty())
+	{
+		// Nothing ready yet
+		return;
+	}
+
+	if (bIsProcessingPSODeferredCreate)
+	{
+		// A deferred create job is still in flight. Ready components will be picked up
+		// when OnDeferredCreateRenderStateEnd_GameThread calls CheckPSODeferredReadiness again.
+		return;
+	}
+
+	// Process whatever is ready now, even if some components are still pending.
+	// This allows incremental proxy creation as PSOs complete, utilizing the per-frame
+	// time budget instead of waiting idle then creating all at once.
+	check(PSODeferredCreateState.ComponentsToProcess.IsEmpty());
+	PSODeferredCreateState.Reset();
+	PSODeferredCreateState.ComponentsToProcess = MoveTemp(PSODeferredReadyComponents);
+
+	// Don't unregister from ContainersWithPSODeferredWork here -- this function can be
+	// called from inside the WorldSubsystem::Tick polling loop, and removing from the set
+	// while iterating would invalidate the iterator. The poll loop safely removes containers
+	// via It.RemoveCurrent() when HasPSODeferredComponents() returns false.
+	if (RegistrationState == ERegistrationState::Registering)
+	{
+		StartDeferredCreateRenderStateWork();
+	}
+}
+
+void UFastGeoContainer::ForceCreateDeferredComponents()
+{
+	check(IsInGameThread());
+
+	// Early-out if no deferred work remains. This can happen when the drain loop
+	// (TickRenderWork) inside this function or the caller's TickRenderWork already
+	// completed all deferred work before we get here.
+	if (!HasPSODeferredComponents())
+	{
+		return;
+	}
+
+	// Drain any in-flight deferred create before overwriting PSODeferredCreateState.
+	// A previous deferred create pass may still be executing (budget-driven re-enqueue
+	// from before the blocking wait started). TickRenderWork handles both async (queue)
+	// and sync (TickDeferredCreate_Sync) modes.
+	while (bIsProcessingPSODeferredCreate)
+	{
+		TickRenderWork(/*bWaitForCompletion=*/ true);
+	}
+	check(PSODeferredCreateState.ComponentsToProcess.IsEmpty());
+
+	// Blocking path (e.g., BlockTillLevelStreamingCompleted): GT is blocked and
+	// cannot pump the task graph, so PSO callbacks will never fire. Merge all
+	// deferred components (pending + ready) into a force-create pass. The PSO
+	// delay check is skipped via IsWaitingForCompletion() in CreateSceneProxy,
+	// ensuring AddToWorld + blocking wait ends with all proxies created.
+	check(PSODeferredCreateState.ComponentsToProcess.IsEmpty());
+	PSODeferredCreateState.Reset();
+	PSODeferredCreateState.ComponentsToProcess = MoveTemp(PSODeferredReadyComponents);
+	PSODeferredCreateState.ComponentsToProcess.Append(PSODeferredPendingComponents.Array());
+	PSODeferredPendingComponents.Reset();
+	GetWorldSubsystem()->UnregisterContainerWithPSODeferredWork(this);
+	StartDeferredCreateRenderStateWork();
+}
+
+void UFastGeoContainer::StartDeferredCreateRenderStateWork()
+{
+	check(IsInGameThread());
+	bIsProcessingPSODeferredCreate = true;
+	if (FFastGeoStreamingModule::IsAsyncRenderWorkAllowed())
+	{
+		UFastGeoWorldSubsystem* WorldSubsystem = GetWorldSubsystem();
+		WorldSubsystem->PushAsyncDeferredCreateJob(this);
+	}
+	// Sync mode: TickRenderWork_Sync handles via TickDeferredCreate_Sync
+}
+
+void UFastGeoContainer::OnDeferredCreateRenderStateBegin_GameThread()
+{
+	check(IsInGameThread());
+	PSODeferredCreateState.NumToProcess = PSODeferredCreateState.ComponentsToProcess.Num() - PSODeferredCreateState.TotalNumProcessed;
+	PSODeferredCreateState.bIsInBlockingWait = GetWorldSubsystem()->IsWaitingForCompletion();
+	AsyncContext.UpdateTimeSeconds(GetWorld()->GetTimeSeconds());
+}
+
+void UFastGeoContainer::OnDeferredCreateRenderState_Concurrent()
+{
+	TRACE_CPUPROFILER_EVENT_SCOPE(UFastGeoContainer::OnDeferredCreateRenderState_Concurrent);
+	ExecuteCreateRenderState_Concurrent<false>(PSODeferredCreateState);
+}
+
+void UFastGeoContainer::OnDeferredCreateRenderStateEnd_GameThread()
+{
+	check(IsInGameThread());
+	PSODeferredCreateState.TotalNumProcessed += PSODeferredCreateState.NumProcessed;
+	PSODeferredCreateState.NumToProcess = 0;
+	PSODeferredCreateState.NumProcessed = 0;
+	PSODeferredCreateState.bIsInBlockingWait = false;
+	bIsProcessingPSODeferredCreate = false;
+
+	if (!PSODeferredCreateState.IsCompleted())
+	{
+		// Budget exhausted
+		StartDeferredCreateRenderStateWork();
+	}
+	else
+	{
+		PSODeferredCreateState.Reset();
+		if (HasPSODeferredComponents())
+		{
+			// New deferred components arrived while the previous pass was in flight
+			// (e.g., a new batch of components completed their initial create pass and
+			// parked while this deferred create was executing). Process them now.
+			// Poll pending to detect any whose PSO callback may have been lost.
+			CheckPSODeferredReadiness(/*bPollPending=*/ true);
+		}
+		else
+		{
+			TryCompleteRegistration();
 		}
 	}
 }
 
+bool UFastGeoContainer::TickDeferredCreate_Sync(const UE::FTimeout& Timeout)
+{
+	if (PSODeferredCreateState.ComponentsToProcess.IsEmpty())
+	{
+		return true;
+	}
+
+	TRACE_CPUPROFILER_EVENT_SCOPE(UFastGeoContainer::TickDeferredCreate_Sync);
+	if (PSODeferredCreateState.ComponentsToProcess.Num() > 0 && PSODeferredCreateState.NumToProcess == 0)
+	{
+		OnDeferredCreateRenderStateBegin_GameThread();
+	}
+
+	const bool bCompleted = FastGeo::AdvanceRenderStateBudgeted(PSODeferredCreateState, Timeout, [](FFastGeoComponent* Component)
+	{
+		TRACE_CPUPROFILER_EVENT_SCOPE(FFastGeoComponent_CreateRenderState);
+		Component->CreateRenderState(nullptr);
+	});
+
+	if (bCompleted)
+	{
+		FastGeo::Private::ResetBatchForSyncCompletion(PSODeferredCreateState);
+		OnDeferredCreateRenderStateEnd_GameThread();
+	}
+	return bCompleted;
+}
+
+void UFastGeoContainer::RoutePSORecreateComponent(FFastGeoComponent* Component)
+{
+	check(Component->NeedsPSORecreate());
+	if (FFastGeoStreamingModule::IsAsyncRenderWorkAllowed())
+	{
+		// Async mode: push to per-container recreate queue (processed by the async job queue).
+		PushPSORecreateComponent(Component);
+	}
+	else
+	{
+		// Sync mode: route through ProcessPendingRecreate (runs every frame from OnWorldPreSendAllEndOfFrameUpdates).
+		Component->MarkRenderStateDirty(false);
+	}
+}
+
+void UFastGeoContainer::PushPSORecreateComponent(FFastGeoComponent* Component)
+{
+	check(IsInGameThread());
+	// Async-only path. Sync mode routes through MarkRenderStateDirty -> ProcessPendingRecreate.
+	check(FFastGeoStreamingModule::IsAsyncRenderWorkAllowed());
+
+	// Don't accept recreate work during unregistration -- the destroy job will clean up all proxies.
+	if (RegistrationState == ERegistrationState::Unregistering)
+	{
+		return;
+	}
+
+	PSORecreateComponents.Add(Component);
+	if (!bIsProcessingPSORecreate)
+	{
+		StartRecreateRenderStateWork();
+	}
+}
+
+void UFastGeoContainer::StartRecreateRenderStateWork()
+{
+	check(IsInGameThread());
+
+	if (!PSORecreateState.IsCompleted())
+	{
+		// Continuing an in-progress pass -- merge new arrivals without resetting progress.
+		// NumToProcess is recalculated in OnRecreateRenderStateBegin_GameThread from Num - TotalNumProcessed.
+		check(bIsProcessingPSORecreate);
+		if (!PSORecreateComponents.IsEmpty())
+		{
+			PSORecreateState.ComponentsToProcess.Append(PSORecreateComponents);
+			PSORecreateComponents.Reset();
+		}
+	}
+	else
+	{
+		// Fresh pass
+		check(!bIsProcessingPSORecreate);
+		bIsProcessingPSORecreate = true;
+		PSORecreateState.Reset();
+		PSORecreateState.ComponentsToProcess = MoveTemp(PSORecreateComponents);
+	}
+
+	if (FFastGeoStreamingModule::IsAsyncRenderWorkAllowed())
+	{
+		UFastGeoWorldSubsystem* WorldSubsystem = GetWorldSubsystem();
+		WorldSubsystem->RegisterContainerWithAsyncPSORecreateWork(this);
+		WorldSubsystem->PushAsyncRecreateRenderStateJob(this);
+	}
+	// Sync mode routes through MarkRenderStateDirty -> ProcessPendingRecreate instead.
+}
+
+void UFastGeoContainer::OnRecreateRenderStateBegin_GameThread()
+{
+	check(IsInGameThread());
+	PSORecreateState.NumToProcess = PSORecreateState.ComponentsToProcess.Num() - PSORecreateState.TotalNumProcessed;
+	PSORecreateState.bIsInBlockingWait = GetWorldSubsystem()->IsWaitingForCompletion();
+	AsyncContext.UpdateTimeSeconds(GetWorld()->GetTimeSeconds());
+}
+
+void UFastGeoContainer::OnRecreateRenderState_Concurrent()
+{
+	TRACE_CPUPROFILER_EVENT_SCOPE(UFastGeoContainer::OnRecreateRenderState_Concurrent);
+	ExecuteCreateRenderState_Concurrent<true>(PSORecreateState);
+}
+
+void UFastGeoContainer::OnRecreateRenderStateEnd_GameThread()
+{
+	check(IsInGameThread());
+	check(bIsProcessingPSORecreate);
+	PSORecreateState.TotalNumProcessed += PSORecreateState.NumProcessed;
+	PSORecreateState.NumToProcess = 0;
+	PSORecreateState.NumProcessed = 0;
+	PSORecreateState.bIsInBlockingWait = false;
+
+	// Container is being unregistered -- abort all recreate work.
+	// The destroy job in the pipe will clean up all proxies.
+	if (RegistrationState == ERegistrationState::Unregistering)
+	{
+		PSORecreateState.Reset();
+		PSORecreateComponents.Reset();
+		bIsProcessingPSORecreate = false;
+		GetWorldSubsystem()->UnregisterContainerWithAsyncPSORecreateWork(this);
+		// Recreate may have been the last pending work -- check if unregistration can complete.
+		TryCompleteUnregistration();
+		return;
+	}
+
+	if (!PSORecreateState.IsCompleted())
+	{
+		StartRecreateRenderStateWork();
+	}
+	else
+	{
+		PSORecreateState.Reset();
+		bIsProcessingPSORecreate = false;
+
+		if (!PSORecreateComponents.IsEmpty())
+		{
+			StartRecreateRenderStateWork();
+		}
+		else
+		{
+			GetWorldSubsystem()->UnregisterContainerWithAsyncPSORecreateWork(this);
+		}
+	}
+}
+#endif
+
 void UFastGeoContainer::OnDestroyRenderStateEnd_GameThread()
 {
 	if (FFastGeoStreamingModule::IsAsyncRenderWorkAllowed())
@@ -1202,7 +1668,7 @@ void UFastGeoContainer::OnDestroyRenderStateEnd_GameThread()
 	PendingDestroy.RenderState.TotalNumProcessed += PendingDestroy.RenderState.NumProcessed;
 	PendingDestroy.RenderState.NumToProcess = 0;
 	PendingDestroy.RenderState.NumProcessed = 0;	
-	PendingDestroy.bIsInBlockingWait = false;
+	PendingDestroy.RenderState.bIsInBlockingWait = false;
 
 	if (!IsDestroyRenderStateCompleted())
 	{
@@ -1212,10 +1678,7 @@ void UFastGeoContainer::OnDestroyRenderStateEnd_GameThread()
 	else
 	{
 		// Render work complete - check if all unregistration systems are done
-		if (!HasAnyPendingDestroyTasks())
-		{
-			OnUnregisterCompleted();
-		}
+		TryCompleteUnregistration();
 	}
 }
 
@@ -1575,7 +2038,8 @@ void UFastGeoContainer::PrecachePSOs()
 	// This function is currently only starting precaching when running in async mode
 	TRACE_CPUPROFILER_EVENT_SCOPE(UFastGeoContainer::PrecachePSOs_GameThread);
 	StartPrecachePSOWork();
-	if (HasAnyPendingTasks())
+	// Recreate is unrelated to PSO precache setup.
+	if (HasAnyPendingRegistrationTasks())
 	{
 		// Only TickRenderWork is needed here to Launch the PSO precache job via ProcessAsyncRenderStateJobs.
 		constexpr bool bWaitForCompletion = false;
@@ -1897,10 +2361,7 @@ void UFastGeoContainer::OnAsyncCreatePhysicsStateEnd_GameThread()
 	RegisterBodyInstances();
 
 	// Physics work complete - check if all registration systems are done
-	if (!HasAnyPendingCreateTasks())
-	{
-		OnRegisterCompleted();
-	}
+	TryCompleteRegistration();
 }
 
 void UFastGeoContainer::OnAsyncDestroyPhysicsStateBegin_GameThread(TSet<UObject*>& OutRootedObjects)
@@ -1950,10 +2411,7 @@ void UFastGeoContainer::OnAsyncDestroyPhysicsStateEnd_GameThread()
 	}
 
 	// Physics work complete - check if all unregistration systems are done
-	if (!HasAnyPendingDestroyTasks())
-	{
-		OnUnregisterCompleted();
-	}
+	TryCompleteUnregistration();
 }
 
 void UFastGeoContainer::CollectBodySetupsWithPhysicsMeshesToCreate(TSet<UBodySetup*>& OutBodySetups) const
diff --git a/Engine/Plugins/Experimental/FastGeoStreaming/Source/FastGeoStreaming/Private/FastGeoDecalComponent.cpp b/Engine/Plugins/Experimental/FastGeoStreaming/Source/FastGeoStreaming/Private/FastGeoDecalComponent.cpp
--- a/Engine/Plugins/Experimental/FastGeoStreaming/Source/FastGeoStreaming/Private/FastGeoDecalComponent.cpp
+++ b/Engine/Plugins/Experimental/FastGeoStreaming/Source/FastGeoStreaming/Private/FastGeoDecalComponent.cpp
@@ -113,7 +113,7 @@ void FFastGeoDecalComponent::CreateRenderState(FRegisterComponentContext* Contex
 	}
 
 	FWriteScopeLock WriteLock(*Lock.Get());
-	check(ProxyState == EProxyCreationState::None);
+	check(ProxyState == EProxyCreationState::None || ProxyState == EProxyCreationState::Delayed);
 	check(!SceneProxy);
 	ProxyState = EProxyCreationState::Creating;
 	bRenderStateDirty = false;
@@ -173,7 +173,9 @@ FDeferredDecalProxy* FFastGeoDecalComponent::CreateSceneProxy(ESceneProxyCreatio
 	check(!SceneProxy);
 
 #if !WITH_EDITOR && UE_WITH_PSO_PRECACHING
-	if (PSOPrecacheComponentData.IsPSOPrecaching() && GetPSOPrecacheProxyCreationStrategy() != EPSOPrecacheProxyCreationStrategy::AlwaysCreate)
+	// Skip PSO delay during blocking wait -- GT cannot pump tasks so PSO callbacks won't fire.
+	const bool bForceCreate = GetWorld()->GetSubsystemChecked<UFastGeoWorldSubsystem>()->IsWaitingForCompletion();
+	if (PSOPrecacheComponentData.IsPSOPrecaching() && GetPSOPrecacheProxyCreationStrategy() != EPSOPrecacheProxyCreationStrategy::AlwaysCreate && !bForceCreate)
 	{
 		OutError = ESceneProxyCreationError::WaitingPSOs;
 		return nullptr;
@@ -267,12 +269,14 @@ void FFastGeoDecalComponent::PrecachePSOs_Concurrent()
 	TArray<FMaterialPSOPrecacheRequestID> MaterialPSOPrecacheRequestIDs;
 	FGraphEventArray PSOPrecacheCompileEvents = SceneProxyDesc.DecalMaterial->PrecachePSOs(VertexFactoryDataList, PSOPrecacheParams, EPSOPrecachePriority::High, MaterialPSOPrecacheRequestIDs);
 
-	// Request recreate of the render state when the PSO compilation is ready (if we want to delay proxy creation)
 	auto OnPrecacheFinishedGT = [this]()
 	{
 		if (GetWorld())
 		{
-			MarkRenderStateDirty(false);
+			if (UFastGeoContainer* Container = GetOwnerContainer())
+			{
+				Container->OnComponentPSOPrecacheCompleted(this);
+			}
 		}
 	};
 
diff --git a/Engine/Plugins/Experimental/FastGeoStreaming/Source/FastGeoStreaming/Private/FastGeoLocalLightComponent.cpp b/Engine/Plugins/Experimental/FastGeoStreaming/Source/FastGeoStreaming/Private/FastGeoLocalLightComponent.cpp
--- a/Engine/Plugins/Experimental/FastGeoStreaming/Source/FastGeoStreaming/Private/FastGeoLocalLightComponent.cpp
+++ b/Engine/Plugins/Experimental/FastGeoStreaming/Source/FastGeoStreaming/Private/FastGeoLocalLightComponent.cpp
@@ -85,7 +85,7 @@ void FFastGeoLocalLightComponent::CreateRenderState(FRegisterComponentContext* C
 	}
 
 	FWriteScopeLock WriteLock(*Lock.Get());
-	check(ProxyState == EProxyCreationState::None);
+	check(ProxyState == EProxyCreationState::None || ProxyState == EProxyCreationState::Delayed);
 	ProxyState = EProxyCreationState::Creating;
 	bRenderStateDirty = false;
 
@@ -119,7 +119,9 @@ FLightSceneProxy* FFastGeoLocalLightComponent::CreateSceneProxy(ESceneProxyCreat
 	check(ShouldComponentAddToRenderScene());
 
 #if !WITH_EDITOR && UE_WITH_PSO_PRECACHING
-	if (PSOPrecacheComponentData.IsPSOPrecaching() && GetPSOPrecacheProxyCreationStrategy() != EPSOPrecacheProxyCreationStrategy::AlwaysCreate)
+	// Skip PSO delay during blocking wait -- GT cannot pump tasks so PSO callbacks won't fire.
+	const bool bForceCreate = GetWorld()->GetSubsystemChecked<UFastGeoWorldSubsystem>()->IsWaitingForCompletion();
+	if (PSOPrecacheComponentData.IsPSOPrecaching() && GetPSOPrecacheProxyCreationStrategy() != EPSOPrecacheProxyCreationStrategy::AlwaysCreate && !bForceCreate)
 	{
 		OutError = ESceneProxyCreationError::WaitingPSOs;
 		return nullptr;
@@ -208,12 +210,14 @@ void FFastGeoLocalLightComponent::PrecachePSOs_Concurrent()
 	FLocalLightSceneProxyDesc& SceneProxyDesc = GetLocalLightSceneProxyDesc();
 	FGraphEventArray PSOPrecacheCompileEvents = SceneProxyDesc.LightFunctionMaterial->PrecachePSOs(VertexFactoryDataList, PSOPrecacheParams, EPSOPrecachePriority::High, MaterialPSOPrecacheRequestIDs);
 
-	// Request recreate of the render state when the PSO compilation is ready (if we want to delay proxy creation)
 	auto OnPrecacheFinishedGT = [this]()
 	{
 		if (GetWorld())
 		{
-			MarkRenderStateDirty(false);
+			if (UFastGeoContainer* Container = GetOwnerContainer())
+			{
+				Container->OnComponentPSOPrecacheCompleted(this);
+			}
 		}
 	};
 
diff --git a/Engine/Plugins/Experimental/FastGeoStreaming/Source/FastGeoStreaming/Private/FastGeoPrimitiveComponent.cpp b/Engine/Plugins/Experimental/FastGeoStreaming/Source/FastGeoStreaming/Private/FastGeoPrimitiveComponent.cpp
--- a/Engine/Plugins/Experimental/FastGeoStreaming/Source/FastGeoStreaming/Private/FastGeoPrimitiveComponent.cpp
+++ b/Engine/Plugins/Experimental/FastGeoStreaming/Source/FastGeoStreaming/Private/FastGeoPrimitiveComponent.cpp
@@ -5,7 +5,6 @@
 #include "FastGeoComponentCluster.h"
 #include "FastGeoInstancedStaticMeshComponent.h"
 #include "FastGeoWorldSubsystem.h"
-#include "FastGeoComponentCluster.h"
 #include "FastGeoSurrogateBodyInstanceIndex.h"
 #include "FastGeoSurrogateComponent.h"
 #include "FastGeoSurrogateActor.h"
@@ -25,7 +24,6 @@
 #include "PrimitiveComponentHelper.h"
 #include "PSOPrecache.h"
 #include "PSOPrecacheMaterial.h"
-#include "PrimitiveComponentHelper.h"
 #include "UnrealEngine.h"
 
 #if WITH_EDITOR
@@ -329,7 +327,8 @@ void FFastGeoPrimitiveComponent::CreateRenderState(FRegisterComponentContext* Co
 	}
 
 	FWriteScopeLock WriteLock(*Lock.Get());
-	check(ProxyState == EProxyCreationState::None);
+	// Delayed is valid: PSO-deferred components re-enter CreateRenderState once PSOs are ready.
+	check(ProxyState == EProxyCreationState::None || ProxyState == EProxyCreationState::Delayed);
 	ProxyState = EProxyCreationState::Creating;
 	bRenderStateDirty = false;
 
@@ -349,6 +348,9 @@ void FFastGeoPrimitiveComponent::CreateRenderState(FRegisterComponentContext* Co
 		Scene->AddPrimitive(&Desc);
 
 		ProxyState = EProxyCreationState::Created;
+#if !WITH_EDITOR && UE_WITH_PSO_PRECACHING
+		bCreatedWithPSOFallbackMaterial = GetSceneProxyDesc().bUsePSOPrecacheFallbackMaterial;
+#endif
 	}
 	else if (Error == ESceneProxyCreationError::WaitingPSOs)
 	{
@@ -387,6 +389,9 @@ void FFastGeoPrimitiveComponent::DestroyRenderState(FFastGeoDestroyRenderStateCo
 		ProxyState = EProxyCreationState::None;
 
 		bRenderStateDirty = false;
+#if !WITH_EDITOR && UE_WITH_PSO_PRECACHING
+		bCreatedWithPSOFallbackMaterial = false;
+#endif
 
 #if WITH_EDITOR
 		NotifyRenderStateChanged();
@@ -649,6 +654,12 @@ void FFastGeoPrimitiveComponent::InitializeDynamicProperties()
 bool FFastGeoPrimitiveComponent::UsePSOPrecacheFallbackMaterial() const
 {
 #if !WITH_EDITOR && UE_WITH_PSO_PRECACHING
+	// Skip fallback material during blocking wait -- force real material creation.
+	// GetWorld() is guaranteed valid during CreateSceneProxy -- the container initiated this work.
+	if (GetWorld()->GetSubsystemChecked<UFastGeoWorldSubsystem>()->IsWaitingForCompletion())
+	{
+		return false;
+	}
 	return PSOPrecacheComponentData.IsPSOPrecaching() && GetPSOPrecacheProxyCreationStrategy() == EPSOPrecacheProxyCreationStrategy::UseFallbackMaterialUntilPSOPrecached;
 #else
 	return false;
@@ -665,6 +676,15 @@ void FFastGeoPrimitiveComponent::PrecachePSOs()
 bool FFastGeoPrimitiveComponent::CheckPSOPrecachingAndBoostPriority(EPSOPrecachePriority NewPSOPrecachePriority)
 {
 #if !WITH_EDITOR && UE_WITH_PSO_PRECACHING
+	// When the world subsystem is in a blocking wait (e.g., BlockTillLevelStreamingCompleted),
+	// skip the PSO delay check to create proxies immediately regardless of PSO compilation
+	// state. GT is blocked and PSO callbacks cannot fire, so we force proxy creation to
+	// match the AddToWorld contract (all proxies created when blocking completes).
+	// GetWorld() is guaranteed valid during CreateSceneProxy -- the container initiated this work.
+	if (GetWorld()->GetSubsystemChecked<UFastGeoWorldSubsystem>()->IsWaitingForCompletion())
+	{
+		return false;
+	}
 	return PSOPrecacheComponentData.CheckPSOPrecachingAndBoostPriority(NewPSOPrecachePriority);
 #else
 	return false;
@@ -720,12 +740,14 @@ void FFastGeoPrimitiveComponent::PrecachePSOs_Concurrent()
 	TArray<FMaterialPSOPrecacheRequestID> MaterialPSOPrecacheRequestIDs;
 	PrecacheMaterialPSOs(PSOPrecacheDataArray, MaterialPSOPrecacheRequestIDs, PSOPrecacheCompileEvents);
 
-	// Request recreate of the render state when the PSO compilation is ready (if we want to delay proxy creation)
 	auto OnPrecacheFinishedGT = [this]()
 	{
 		if (GetWorld())
 		{
-			MarkRenderStateDirty();
+			if (UFastGeoContainer* Container = GetOwnerContainer())
+			{
+				Container->OnComponentPSOPrecacheCompleted(this);
+			}
 		}
 	};
 
diff --git a/Engine/Plugins/Experimental/FastGeoStreaming/Source/FastGeoStreaming/Private/FastGeoWorldSubsystem.cpp b/Engine/Plugins/Experimental/FastGeoStreaming/Source/FastGeoStreaming/Private/FastGeoWorldSubsystem.cpp
--- a/Engine/Plugins/Experimental/FastGeoStreaming/Source/FastGeoStreaming/Private/FastGeoWorldSubsystem.cpp
+++ b/Engine/Plugins/Experimental/FastGeoStreaming/Source/FastGeoStreaming/Private/FastGeoWorldSubsystem.cpp
@@ -38,6 +38,13 @@ namespace FastGeo
 		TEXT("Maximum time budget in milliseconds for the async render state tasks (0 = no time limit)"),
 		ECVF_Default);
 
+	static float GSyncPSORecreateRenderStateTimeBudgetMS = 0.0f;
+	static FAutoConsoleVariableRef CVarSyncPSORecreateRenderStateTimeBudgetMS(
+		TEXT("FastGeo.SyncPSORecreateRenderState.TimeBudgetMS"),
+		GSyncPSORecreateRenderStateTimeBudgetMS,
+		TEXT("Maximum time budget in milliseconds for PSO recreate render state work per frame in sync mode (0 = no time limit)"),
+		ECVF_Default);
+
 	static int32 GAsyncRenderStateTaskMaxNumComponentsToProcess = 0;
 	static FAutoConsoleVariableRef CVarAsyncRenderStateTaskMaxNumComponentsToProcess(
 		TEXT("FastGeo.AsyncRenderStateTask.MaxNumComponentsToProcess"),
@@ -496,12 +503,14 @@ void UFastGeoWorldSubsystem::OnRemoveLevelFromWorldExtension(ULevel* InLevel, co
 	{
 		TRACE_CPUPROFILER_EVENT_SCOPE(UFastGeoWorldSubsystem::OnRemoveLevelFromWorldExtension);
 
-		if (FastGeo->HasAnyPendingDestroyTasks())
+		// Drain all pending work including destroy and PSO recreate (UseFallbackMaterialUntilPSOPrecached). 
+		// Kept in sync with CheckNoPendingTasks which validates HasAnyPendingTasks() after removal completes.
+		if (FastGeo->HasAnyPendingTasks())
 		{
 			FastGeo->Tick(bInWaitForCompletion);
 		}
 
-		if (FastGeo->HasAnyPendingDestroyTasks())
+		if (FastGeo->HasAnyPendingTasks())
 		{
 			bOutHasCompleted = false;
 		}
@@ -520,6 +529,7 @@ void UFastGeoWorldSubsystem::OnUpdateLevelStreaming()
 void UFastGeoWorldSubsystem::WaitForAllPendingWorkCompletion()
 {
 	check(IsInGameThread());
+	TGuardValue<bool> GuardWaitingForCompletion(bWaitingForCompletion, true);
 
 	static constexpr bool bWaitForCompletion = true;
 
@@ -542,7 +552,7 @@ void UFastGeoWorldSubsystem::WaitForAllPendingWorkCompletion()
 
 		// Drain recreate queue
 		ProcessPendingRecreate();
-		if (ComponentsPendingRecreate.IsEmpty())
+		if (PendingRecreateManager.IsEmpty())
 		{
 			break;
 		}
@@ -561,49 +571,132 @@ void UFastGeoWorldSubsystem::WaitForAllPendingWorkCompletion()
 	FPhysScene* PhysScene = World ? World->GetPhysicsScene() : nullptr;
 	check(!PhysScene || !PhysScene->HasAsyncPhysicsStateJobs());
 	check(!AsyncRenderStateJobQueue.IsValid() || AsyncRenderStateJobQueue->IsCompleted());
-	check(ComponentsPendingRecreate.IsEmpty());
+	check(PendingRecreateManager.IsEmpty());
 	check(PendingRuntimeRegisterContainers.IsEmpty());
 	check(PendingRuntimeDestroyContainers.IsEmpty());
 }
 
-void UFastGeoWorldSubsystem::AddToComponentsPendingRecreate(FFastGeoComponent* InComponentPendingRecreate)
+void UFastGeoWorldSubsystem::AddToComponentsPendingRecreate(FFastGeoComponent* InComponent)
 {
-	ComponentsPendingRecreate.Emplace(InComponentPendingRecreate);
+	PendingRecreateManager.Add(InComponent);
 }
 
 void UFastGeoWorldSubsystem::ProcessPendingRecreate()
 {
 	TRACE_CPUPROFILER_EVENT_SCOPE(UFastGeoWorldSubsystem::ProcessPendingRecreate);
 
-	CSV_CUSTOM_STAT(FastGeo, PendingRecreate, ComponentsPendingRecreate.Num(), ECsvCustomStatOp::Set);
+	if (PendingRecreateManager.IsEmpty())
+	{
+		return;
+	}
 
 #if WITH_EDITOR
-	if (!ComponentsPendingRecreate.IsEmpty())
+	if (!PendingRecreateManager.GetComponents().IsEmpty())
+	{
+		ComponentsPreRecreateEvent.Broadcast(PendingRecreateManager.GetComponents());
+	}
+#endif
+
+	PendingRecreateManager.Process(GetWorld());
+}
+
+bool UFastGeoWorldSubsystem::FFastGeoPendingRecreateManager::IsEmpty() const
+{
+	if (!Components.IsEmpty())
+	{
+		return false;
+	}
+#if !WITH_EDITOR && UE_WITH_PSO_PRECACHING
+	if (!PSORecreateBatch.IsCompleted())
 	{
-		ComponentsPreRecreateEvent.Broadcast(ComponentsPendingRecreate);
+		return false;
 	}
 #endif
+	return true;
+}
+
+void UFastGeoWorldSubsystem::FFastGeoPendingRecreateManager::Reset()
+{
+	Components.Reset();
+#if !WITH_EDITOR && UE_WITH_PSO_PRECACHING
+	PSORecreateBatch.Reset();
+#endif
+	bHasNonPSORecreate = false;
+}
+
+void UFastGeoWorldSubsystem::FFastGeoPendingRecreateManager::Add(FFastGeoComponent* Component)
+{
+	Components.Emplace(Component);
+#if !WITH_EDITOR && UE_WITH_PSO_PRECACHING
+	if (!Component->NeedsPSORecreate())
+#endif
+	{
+		bHasNonPSORecreate = true;
+	}
+}
 
-	int32 NumPrimitiveDelayed = 0;
-	for (const FFastGeoRegisteredComponent& PendingRecreateComponent : ComponentsPendingRecreate)
+void UFastGeoWorldSubsystem::FFastGeoPendingRecreateManager::Process(UWorld* World)
+{
+	CSV_CUSTOM_STAT(FastGeo, PendingRecreate, Components.Num(), ECsvCustomStatOp::Set);
+
+	// Pass 1: Process non-PSO recreates immediately (preserve same-frame guarantee).
+	// Skipped when only PSO recreates are pending (common case during PSO completion bursts).
+	if (bHasNonPSORecreate)
 	{
-		// Skip component if its container was unregistered or re-registered before we process it
-		if (FFastGeoComponent* Component = PendingRecreateComponent.TryGetRegistered<FFastGeoComponent>())
+		bHasNonPSORecreate = false;
+		for (int32 i = Components.Num() - 1; i >= 0; --i)
 		{
-#if CSV_PROFILER && !CSV_PROFILER_MINIMAL
-			if (Component->IsRenderStateDelayed())
+			if (FFastGeoComponent* Component = Components[i].TryGetRegistered<FFastGeoComponent>())
 			{
-				++NumPrimitiveDelayed;
+#if !WITH_EDITOR && UE_WITH_PSO_PRECACHING
+				if (Component->NeedsPSORecreate())
+				{
+					continue;
 			}
 #endif
 			Component->DestroyRenderState(nullptr);
 			Component->CreateRenderState(nullptr);
 		}
+			Components.RemoveAtSwap(i, EAllowShrinking::No);
+		}
+	}
+
+#if !WITH_EDITOR && UE_WITH_PSO_PRECACHING
+	// Pass 2: PSO recreates with ParallelFor + time budget.
+	// Uses TFastGeoRenderStateBatch<FFastGeoRegisteredComponent> to persist across frames
+	// without raw pointers. Stale components (container unregistered between frames) are
+	// validated inline via TryGetRegistered and skipped at zero cost.
+	if (!Components.IsEmpty())
+	{
+		PSORecreateBatch.ComponentsToProcess.Append(MoveTemp(Components));
+		PSORecreateBatch.NumToProcess = PSORecreateBatch.ComponentsToProcess.Num() - PSORecreateBatch.TotalNumProcessed;
 	}
 
-	CSV_CUSTOM_STAT(FastGeo, PendingRecreateDelayed, NumPrimitiveDelayed, ECsvCustomStatOp::Set);
+	if (!PSORecreateBatch.IsCompleted())
+	{
+		// Bypass budget when draining all pending work to avoid the safety iteration limit.
+		const bool bForceComplete = World && World->GetSubsystemChecked<UFastGeoWorldSubsystem>()->IsWaitingForCompletion();
+		const UE::FTimeout Timeout = (!bForceComplete && FastGeo::GSyncPSORecreateRenderStateTimeBudgetMS > 0.0f) ? UE::FTimeout(FastGeo::GSyncPSORecreateRenderStateTimeBudgetMS / 1000.0) : UE::FTimeout::Never();
+
+		FastGeo::AdvanceRenderStateBudgeted(PSORecreateBatch, Timeout, [](FFastGeoRegisteredComponent& Ref)
+		{
+			if (FFastGeoComponent* Component = Ref.TryGetRegistered<FFastGeoComponent>())
+			{
+				Component->DestroyRenderState(nullptr);
+				Component->CreateRenderState(nullptr);
+			}
+		});
+
+		CSV_CUSTOM_STAT(FastGeo, PSORecreateBatchRemaining, PSORecreateBatch.ComponentsToProcess.Num() - PSORecreateBatch.TotalNumProcessed, ECsvCustomStatOp::Set);
 
-	ComponentsPendingRecreate.Reset();
+		if (PSORecreateBatch.IsCompleted())
+		{
+			// Use Empty() instead of Reset() to free memory from large PSO bursts.
+			PSORecreateBatch.ComponentsToProcess.Empty();
+			PSORecreateBatch.Reset();
+		}
+	}
+#endif
 }
 
 void UFastGeoWorldSubsystem::ForEachHLODObjectInCell(const UWorldPartitionRuntimeCell* InCell, TFunction<void(IWorldPartitionHLODObject*)> InFunc)
@@ -622,7 +715,7 @@ void UFastGeoWorldSubsystem::ForEachHLODObjectInCell(const UWorldPartitionRuntim
 }
 
 #if DO_CHECK
-void UFastGeoWorldSubsystem::CheckNoPendingTasks(ULevel* Level, UWorld* World)
+void UFastGeoWorldSubsystem::CheckNoPendingTasks(ULevel* Level, UWorld* World, bool bIncludeRecreateTasks)
 {
 	if (World != GetWorld())
 	{
@@ -631,14 +724,19 @@ void UFastGeoWorldSubsystem::CheckNoPendingTasks(ULevel* Level, UWorld* World)
 
 	if (UFastGeoContainer* FastGeo = Level->GetAssetUserData<UFastGeoContainer>())
 	{
-		check(!FastGeo->HasAnyPendingTasks());
+		// bIncludeRecreateTasks=false: recreate (UseFallbackMaterialUntilPSOPrecached) is
+		// post-registration background work that can legitimately be pending after AddToWorld.
+		// bIncludeRecreateTasks=true: after RemoveFromWorld, all work including recreate must be done.
+		const bool bHasPendingTasks = bIncludeRecreateTasks ? FastGeo->HasAnyPendingTasks() : FastGeo->HasAnyPendingRegistrationTasks();
+		check(!bHasPendingTasks);
 	}
 }
 
 void UFastGeoWorldSubsystem::OnLevelAddedToWorld(ULevel* Level, UWorld* World)
 {
 	check(Level);
-	CheckNoPendingTasks(Level, World);
+	// Recreate can be pending after AddToWorld -- it's post-registration background work.
+	CheckNoPendingTasks(Level, World, /*bIncludeRecreateTasks=*/ false);
 }
 
 void UFastGeoWorldSubsystem::OnLevelRemovedFromWorld(ULevel* Level, UWorld* World)
@@ -703,10 +801,66 @@ void UFastGeoWorldSubsystem::Tick(float DeltaTime)
 	// Level-streamed containers are driven by OnAddLevelToWorldExtension/OnRemoveLevelFromWorldExtension instead.
 	// Also evict containers marked for destruction (e.g., DestroyRuntime called from an OnRegistered callback).
 	// These are tracked by PendingRuntimeDestroyContainers and should not be ticked from both sets.
-	TickAndDrainCompleted(PendingRuntimeRegisterContainers, [](const UFastGeoContainer* Container) { return Container->IsFullyRegistered() || Container->bIsBeingAsyncDestroyed; });
+	// Keep containers in the set until registration AND post-registration recreate
+	// (UseFallbackMaterialUntilPSOPrecached) complete. Without this, a standalone runtime
+	// container would be drained after registration, and nobody would tick the async queue
+	// to process the recreate job's End_GameThread callback.
+	TickAndDrainCompleted(PendingRuntimeRegisterContainers, [](const UFastGeoContainer* Container)
+	{
+		bool bRegistrationComplete = Container->IsFullyRegistered();
+#if !WITH_EDITOR && UE_WITH_PSO_PRECACHING
+		bRegistrationComplete = bRegistrationComplete && !Container->HasAnyPendingRecreateTasks();
+#endif
+		return bRegistrationComplete || Container->bIsBeingAsyncDestroyed;
+	});
 
 	// Containers complete destruction when fully unregistered.
 	TickAndDrainCompleted(PendingRuntimeDestroyContainers, [](const UFastGeoContainer* Container) { return Container->IsFullyUnregistered(); });
+
+#if !WITH_EDITOR && UE_WITH_PSO_PRECACHING
+	// Safety net: poll containers with PSO-deferred components to detect lost callbacks.
+	// PSO callbacks can be lost if a component is destroyed before the callback fires, or if
+	// the FPSOPrecacheFinishedTask is never dispatched (lifecycle ID mismatch). Without this
+	// poll, a lost callback would leave the container parked forever with AddToWorld never
+	// completing. The cost is negligible: only containers with active deferred work are polled,
+	// and the check is a simple IsPSOPrecaching() per deferred component.
+	for (auto It = ContainersWithPSODeferredWork.CreateIterator(); It; ++It)
+	{
+		UFastGeoContainer* Container = (*It).Get();
+		if (!IsValid(Container) || !Container->HasPSODeferredComponents())
+		{
+			It.RemoveCurrent();
+		}
+		else
+		{
+			Container->CheckPSODeferredReadiness(/*bPollPending=*/ true);
+		}
+	}
+
+	// Drive async recreate jobs that were pushed after OnAddLevelToWorldExtension completed.
+	// Level-streamed containers are registered and created during AddToWorld, which processes
+	// the async job queue to completion via OnAddLevelToWorldExtension. However, PSO callbacks
+	// (OnComponentPSOPrecacheCompleted) can fire after the extension finishes, pushing recreate
+	// jobs to the queue with no one left to tick them. Runtime containers don't have this
+	// problem because they stay in PendingRuntimeRegisterContainers and get Container->Tick()
+	// each frame. Without this, recreate jobs for level-streamed containers would stall until
+	// the next level streaming event (e.g., camera move streaming in new levels).
+	if (!ContainersWithAsyncPSORecreateWork.IsEmpty())
+	{
+		for (auto It = ContainersWithAsyncPSORecreateWork.CreateIterator(); It; ++It)
+		{
+			if (!IsValid((*It).Get()) || !(*It)->HasAnyPendingRecreateTasks())
+			{
+				It.RemoveCurrent();
+			}
+		}
+
+		if (!ContainersWithAsyncPSORecreateWork.IsEmpty())
+		{
+			ProcessAsyncRenderStateJobs();
+		}
+	}
+#endif
 }
 
 TStatId UFastGeoWorldSubsystem::GetStatId() const
@@ -771,6 +925,32 @@ void UFastGeoWorldSubsystem::PushRuntimeContainerForRegistration(UFastGeoContain
 	}
 }
 
+#if !WITH_EDITOR && UE_WITH_PSO_PRECACHING
+void UFastGeoWorldSubsystem::RegisterContainerWithPSODeferredWork(UFastGeoContainer* FastGeo)
+{
+	check(IsInGameThread());
+	ContainersWithPSODeferredWork.Add(FastGeo);
+}
+
+void UFastGeoWorldSubsystem::UnregisterContainerWithPSODeferredWork(UFastGeoContainer* FastGeo)
+{
+	check(IsInGameThread());
+	ContainersWithPSODeferredWork.Remove(FastGeo);
+}
+
+void UFastGeoWorldSubsystem::RegisterContainerWithAsyncPSORecreateWork(UFastGeoContainer* FastGeo)
+{
+	check(IsInGameThread());
+	ContainersWithAsyncPSORecreateWork.Add(FastGeo);
+}
+
+void UFastGeoWorldSubsystem::UnregisterContainerWithAsyncPSORecreateWork(UFastGeoContainer* FastGeo)
+{
+	check(IsInGameThread());
+	ContainersWithAsyncPSORecreateWork.Remove(FastGeo);
+}
+#endif
+
 void UFastGeoWorldSubsystem::PushRuntimeContainerForDestruction(UFastGeoContainer* FastGeo)
 {
 	check(IsInGameThread());
@@ -824,6 +1004,22 @@ void UFastGeoWorldSubsystem::PushAsyncDestroyRenderStateJob(UFastGeoContainer* F
 	GetOrCreateJobQueue()->AddJob({ FastGeo, FFastGeoAsyncRenderStateJobQueue::EJobType::DestroyRenderState });
 }
 
+#if !WITH_EDITOR && UE_WITH_PSO_PRECACHING
+void UFastGeoWorldSubsystem::PushAsyncDeferredCreateJob(UFastGeoContainer* FastGeo)
+{
+	check(IsInGameThread());
+	check(FastGeo->GetWorld() == GetWorld());
+	GetOrCreateJobQueue()->AddJob({ FastGeo, FFastGeoAsyncRenderStateJobQueue::EJobType::DeferredCreateRenderState });
+}
+
+void UFastGeoWorldSubsystem::PushAsyncRecreateRenderStateJob(UFastGeoContainer* FastGeo)
+{
+	check(IsInGameThread());
+	check(FastGeo->GetWorld() == GetWorld());
+	GetOrCreateJobQueue()->AddJob({ FastGeo, FFastGeoAsyncRenderStateJobQueue::EJobType::RecreateRenderState });
+}
+#endif
+
 void UFastGeoWorldSubsystem::ProcessAsyncRenderStateJobs(bool bWaitForCompletion)
 {
 	check(IsInGameThread());
